Implemented the 'Unsaved' Editor status bar button that displays the number of unsaved assets (dirty packages) the user and potential warnings with respect to the source control those asset may have.
  - Implemented the UnsavedAssetTrackers plugin and enabled it by default.
  - Added the 'Unsaved' button to the Editor taskbars to display the number of unsaved packages along with a warning icon if warning are associated with one of the unsaved files.
  - Updated the 'Save Content' package to display warnings if one of the packages has warnings.
  - Added a class to monitor and periodically update the source controlled status of a collection of files/packages.
  - Changed the default value of PromptForCheckoutOnAssetModification settings to False because the Unsaved button flow is complementary (and the popup was noisy for OFPA users).
  - Updated the changelist window to use the soure control file status monitor.

diff --git a/Engine/Config/BaseEditorPerProjectUserSettings.ini b/Engine/Config/BaseEditorPerProjectUserSettings.ini
--- a/Engine/Config/BaseEditorPerProjectUserSettings.ini
+++ b/Engine/Config/BaseEditorPerProjectUserSettings.ini
@@ -184,7 +184,7 @@ AutoSaveWarningInSeconds=10
 ; Add the game directory by default
 +AutoReimportDirectorySettings=(SourceDirectory="/Game/",MountPoint=,Wildcards=((Wildcard="Localization/*")))
 ; Whether to automatically prompt for SCC checkout on asset modification
-bPromptForCheckoutOnAssetModification=True
+bPromptForCheckoutOnAssetModification=False
 ; Source control
 bSCCAutoAddNewFiles=True
 ; Tool to use for diffing text
diff --git a/Engine/Source/Developer/SourceControl/Private/SourceControlFileStatusMonitor.cpp b/Engine/Source/Developer/SourceControl/Private/SourceControlFileStatusMonitor.cpp
new file mode 100644
--- /dev/null
+++ b/Engine/Source/Developer/SourceControl/Private/SourceControlFileStatusMonitor.cpp
@@ -0,0 +1,297 @@
+// Copyright Epic Games, Inc. All Rights Reserved.
+
+#include "SourceControlFileStatusMonitor.h"
+
+#include "Framework/Application/SlateApplication.h"
+#include "HAL/PlatformTime.h"
+#include "ISourceControlModule.h"
+#include "ISourceControlProvider.h"
+#include "Math/NumericLimits.h"
+#include "SourceControlOperations.h"
+
+FSourceControlFileStatusMonitor::FSourceControlFileStatusMonitor()
+: ProbationPeriodPolicy(FTimespan::FromSeconds(1))
+, RefreshPeriodPolicy(FTimespan::FromMinutes(5))
+{
+	TickerHandle = FTSTicker::GetCoreTicker().AddTicker(FTickerDelegate::CreateRaw(this, &FSourceControlFileStatusMonitor::Tick));
+	SetSuspendMonitoringPolicy([]()
+	{
+		// By default, suspend monitoring if the user didn't interact for the last 5 minutes.
+		return FPlatformTime::Seconds() - FSlateApplication::Get().GetLastUserInteractionTime() > FTimespan::FromMinutes(5).GetTotalSeconds();
+	});
+}
+
+FSourceControlFileStatusMonitor::~FSourceControlFileStatusMonitor()
+{
+	FTSTicker::GetCoreTicker().RemoveTicker(TickerHandle);
+}
+
+TSharedPtr<FSourceControlFileStatusMonitor::FSourceControlFileStatus> FSourceControlFileStatusMonitor::FindFileStatus(const FString& Pathname) const
+{
+	if (const TSharedPtr<FSourceControlFileStatus>* FileStatus = MonitoredFiles.Find(Pathname))
+	{
+		return *FileStatus;
+	}
+	return nullptr;
+}
+
+void FSourceControlFileStatusMonitor::OnSourceControlProviderChanged(ISourceControlProvider& OldProvider, ISourceControlProvider& NewProvider)
+{
+	// Start new with the new provider.
+	RequestedStatusFiles.Reset();
+	for (const TPair<FString, TSharedPtr<FSourceControlFileStatus>>& Pair : MonitoredFiles)
+	{
+		if (Pair.Value->FileState)
+		{
+			for (TPair<uintptr_t, FOnSourceControlFileStatus>& OwnerDelegatePair : Pair.Value->OwnerDelegateMap)
+			{
+				OwnerDelegatePair.Value.ExecuteIfBound(Pair.Key, nullptr); // Passing 'nullptr' state means that the state is now unknown.
+			}
+			Pair.Value->LastStatusCheckTimestampSecs = 0.0;
+			Pair.Value->FileState.Reset();
+		}
+	}
+	NewAddedFileCount = MonitoredFiles.Num();
+	LastAddedFileTimeSecs = FPlatformTime::Seconds();
+	OldestFileStatusTimeSecs = 0.0;
+}
+
+void FSourceControlFileStatusMonitor::StartMonitoringFile(uintptr_t OwnerId, const FString& Pathname, FOnSourceControlFileStatus OnSourceControlFileStatus)
+{
+	// If the file is already monitored.
+	if (TSharedPtr<FSourceControlFileStatus> FileStatus = FindFileStatus(Pathname))
+	{
+		// Another 'client' is looking at this file (if the client already exit, override the callback with the new one).
+		FOnSourceControlFileStatus& OnSourceControlFileStatusDelegate = FileStatus->OwnerDelegateMap.FindOrAdd(OwnerId);
+		OnSourceControlFileStatusDelegate = OnSourceControlFileStatus;
+
+		// The monitor already knows the status.
+		if (FileStatus->FileState)
+		{
+			OnSourceControlFileStatusDelegate.ExecuteIfBound(Pathname, FileStatus->FileState.Get());
+		}
+	}
+	else
+	{
+		MonitoredFiles.Emplace(Pathname, MakeShared<FSourceControlFileStatus>(OwnerId, MoveTemp(OnSourceControlFileStatus)));
+		LastAddedFileTimeSecs = FPlatformTime::Seconds();
+		++NewAddedFileCount;
+	}
+
+	if (!SourceControlProviderChangedDelegateHandle.IsValid())
+	{
+		SourceControlProviderChangedDelegateHandle = ISourceControlModule::Get().RegisterProviderChanged(
+			FSourceControlProviderChanged::FDelegate::CreateSP(this, &FSourceControlFileStatusMonitor::OnSourceControlProviderChanged));
+	}
+}
+
+void FSourceControlFileStatusMonitor::StartMonitoringFiles(uintptr_t OwnerId, const TArray<FString>& Pathnames, FOnSourceControlFileStatus OnSourceControlledFileStatus)
+{
+	for (const FString& Pathname : Pathnames)
+	{
+		StartMonitoringFile(OwnerId, Pathname, OnSourceControlledFileStatus);
+	}
+}
+
+void FSourceControlFileStatusMonitor::StartMonitoringFiles(uintptr_t OwnerId, const TSet<FString>& Pathnames, FOnSourceControlFileStatus OnSourceControlledFileStatus)
+{
+	for (const FString& Pathname : Pathnames)
+	{
+		StartMonitoringFile(OwnerId, Pathname, OnSourceControlledFileStatus);
+	}
+}
+
+
+void FSourceControlFileStatusMonitor::StopMonitoringFile(uintptr_t OwnerId, const FString& Pathname)
+{
+	if (TSharedPtr<FSourceControlFileStatus> FileStatus = FindFileStatus(Pathname))
+	{
+		if (FileStatus->OwnerDelegateMap.Remove(OwnerId) > 0 && FileStatus->OwnerDelegateMap.IsEmpty())
+		{
+			MonitoredFiles.Remove(Pathname);
+		}
+	}
+}
+
+void FSourceControlFileStatusMonitor::StopMonitoringFiles(uintptr_t OwnerId, const TArray<FString>& Pathnames)
+{
+	for (const FString& Pathname : Pathnames)
+	{
+		StopMonitoringFile(OwnerId, Pathname);
+	}
+}
+
+void FSourceControlFileStatusMonitor::StopMonitoringFiles(uintptr_t OwnerId, const TSet<FString>& Pathnames)
+{
+	for (const FString& Pathname : Pathnames)
+	{
+		StopMonitoringFile(OwnerId, Pathname);
+	}
+}
+
+void FSourceControlFileStatusMonitor::StopMonitoringFiles(uintptr_t OwnerId)
+{
+	for (auto It = MonitoredFiles.CreateIterator(); ;++It)
+	{
+		if (It->Value->OwnerDelegateMap.Remove(OwnerId) > 0 && It->Value->OwnerDelegateMap.IsEmpty())
+		{
+			It.RemoveCurrent();
+		}
+	}
+}
+
+TSet<FString> FSourceControlFileStatusMonitor::GetMonitoredFiles(uintptr_t OwnerId)
+{
+	TSet<FString> Pathnames;
+	for (const TPair<FString, TSharedPtr<FSourceControlFileStatus>>& Pair : MonitoredFiles)
+	{
+		if (!Pair.Value->OwnerDelegateMap.Contains(OwnerId))
+		{
+			Pathnames.Add(Pair.Key);
+		}
+	}
+	return Pathnames;
+}
+
+TOptional<FTimespan> FSourceControlFileStatusMonitor::GetStatusAge(const FString& Pathname) const
+{
+	TOptional<FTimespan> Age;
+	if (TSharedPtr<FSourceControlFileStatus> FileStatus = FindFileStatus(Pathname))
+	{
+		Age.Emplace(FTimespan::FromSeconds(FileStatus->FileState ? FPlatformTime::Seconds() - FileStatus->LastStatusCheckTimestampSecs : 0.0));
+	}
+	return Age;
+}
+
+bool FSourceControlFileStatusMonitor::Tick(float DeltaTime)
+{
+	// Nothing to check or a request is already in-flight.
+	if (!ISourceControlModule::Get().IsEnabled() || MonitoredFiles.IsEmpty() || HasOngoingRequest())
+	{
+		return true;
+	}
+
+	// Check if the monitor is suspended, not allowed to send requests.
+	if (SuspendMonitoringPolicy && SuspendMonitoringPolicy())
+	{
+		return true;
+	}
+
+	double NowSecs = FPlatformTime::Seconds();
+
+	// Throttle the source control status check when no new files need to be checked. Don't overload the source control server with too many requests.
+	if (NewAddedFileCount == 0 && NowSecs - OldestFileStatusTimeSecs < RefreshPeriodPolicy.GetTotalSeconds())
+	{
+		return true; // Nothing to do this time around.
+	}
+
+	// Throttle the checks when new files are added. Batch new files edited close in time to be more efficient.
+	if (NowSecs - LastAddedFileTimeSecs < ProbationPeriodPolicy.GetTotalSeconds() && NewAddedFileCount < MaxFileNumPerRequestPolicy)
+	{
+		return true; // Delay the request, give chance to capture more files.
+	}
+
+	TArray<const TPair<FString, TSharedPtr<FSourceControlFileStatus>>*> NewFiles;
+	NewFiles.Reserve(NewAddedFileCount);
+
+	TArray<const TPair<FString, TSharedPtr<FSourceControlFileStatus>>*> RefreshedFiles;
+	RefreshedFiles.Reserve(MonitoredFiles.Num());
+
+	// List all the files that are new and those that weren't updated recently.
+	for (const TPair<FString, TSharedPtr<FSourceControlFileStatus>>& Pair : MonitoredFiles)
+	{
+		if (!Pair.Value->FileState.IsValid())
+		{
+			NewFiles.Add(&Pair);
+		}
+		else if (NowSecs - Pair.Value->LastStatusCheckTimestampSecs >= RefreshPeriodPolicy.GetTotalSeconds())
+		{
+			RefreshedFiles.Add(&Pair);
+		}
+	}
+
+	// Too many status to query/refresh?
+	if (NewFiles.Num() < MaxFileNumPerRequestPolicy && NewFiles.Num() + RefreshedFiles.Num() > MaxFileNumPerRequestPolicy)
+	{
+		// Get the status of all newly added files and refresh the ones that were less recently updated.
+		RefreshedFiles.Sort([](const TPair<FString, TSharedPtr<FSourceControlFileStatus>>& Lhs, const TPair<FString, TSharedPtr<FSourceControlFileStatus>>& Rhs)
+		{
+			// Sort ascending as we are going to use Last()/Pop() later, so the oldest must be at the end.
+			return Lhs.Value->LastStatusCheckTimestampSecs > Rhs.Value->LastStatusCheckTimestampSecs;
+		});
+	}
+
+	RequestedStatusFiles.Reserve(MaxFileNumPerRequestPolicy);
+	while (RequestedStatusFiles.Num() < MaxFileNumPerRequestPolicy)
+	{
+		if (NewFiles.Num())
+		{
+			RequestedStatusFiles.Emplace(NewFiles.Last()->Key);
+			NewFiles.Pop(/*bAllowShrinking*/false);
+			NewAddedFileCount--;
+		}
+		else if (RefreshedFiles.Num())
+		{
+			RequestedStatusFiles.Emplace(RefreshedFiles.Last()->Key);
+			RefreshedFiles.Pop(/*bAllowShrinking*/false);
+		}
+		else
+		{
+			break; // All files to query/refresh were added.
+		}
+	}
+
+	if (RequestedStatusFiles.Num())
+	{
+		LastSourceControlCheckSecs = NowSecs;
+
+		TSharedRef<FUpdateStatus> UpdateStatusRequest = ISourceControlOperation::Create<FUpdateStatus>();
+		UpdateStatusRequest->SetForceUpdate(true);
+
+		ISourceControlModule::Get().GetProvider().Execute(UpdateStatusRequest, RequestedStatusFiles, EConcurrency::Asynchronous,
+			FSourceControlOperationComplete::CreateSP(this, &FSourceControlFileStatusMonitor::OnSourceControlStatusUpdate));
+	}
+
+	return true;
+}
+
+void FSourceControlFileStatusMonitor::OnSourceControlStatusUpdate(const TSharedRef<ISourceControlOperation>& InOperation, ECommandResult::Type InResult)
+{
+	check(IsInGameThread()); // Concurrency issues if the callback in invoked from a background thread.
+
+	double NowSecs = FPlatformTime::Seconds();
+
+	ON_SCOPE_EXIT
+	{
+		RequestedStatusFiles.Reset();
+
+		OldestFileStatusTimeSecs = NowSecs;
+		for (const TPair<FString, TSharedPtr<FSourceControlFileStatus>>& Pair : MonitoredFiles)
+		{
+			OldestFileStatusTimeSecs = FMath::Min(OldestFileStatusTimeSecs, Pair.Value->LastStatusCheckTimestampSecs);
+		}
+	};
+
+	if (InResult != ECommandResult::Succeeded)
+	{
+		return;
+	}
+
+	for (const FString& Pathname : RequestedStatusFiles)
+	{
+		// NOTE: The file and its status can be removed while exeucting OnFileStatusUpdateDelegate. Keep the shared pointer to avoid early destruction.
+		if (TSharedPtr<FSourceControlFileStatus> FileStatus = FindFileStatus(Pathname))
+		{
+			if (TSharedPtr<ISourceControlState> FileState = ISourceControlModule::Get().GetProvider().GetState(Pathname, EStateCacheUsage::Use))
+			{
+				FileStatus->FileState = MoveTemp(FileState);
+				FileStatus->LastStatusCheckTimestampSecs = NowSecs;
+
+				for (TPair<uintptr_t, FOnSourceControlFileStatus>& OwnerDelegatePair : FileStatus->OwnerDelegateMap)
+				{
+					OwnerDelegatePair.Value.ExecuteIfBound(Pathname, FileStatus->FileState.Get());
+				}
+			}
+		}
+	}
+}
diff --git a/Engine/Source/Developer/SourceControl/Private/SourceControlModule.cpp b/Engine/Source/Developer/SourceControl/Private/SourceControlModule.cpp
--- a/Engine/Source/Developer/SourceControl/Private/SourceControlModule.cpp
+++ b/Engine/Source/Developer/SourceControl/Private/SourceControlModule.cpp
@@ -6,6 +6,7 @@
 #include "SourceControlOperations.h"
 #include "SourceControlHelpers.h"
 #include "SourceControlAssetDataCache.h"
+#include "SourceControlFileStatusMonitor.h"
 #include "Misc/Paths.h"
 
 #if SOURCE_CONTROL_WITH_SLATE
@@ -65,6 +66,8 @@ void FSourceControlModule::StartupModule()
 #endif
 
 	AssetDataCache.Startup();
+
+	SourceControlFileStatusMonitor = MakeShared<FSourceControlFileStatusMonitor>();
 }
 
 void FSourceControlModule::ShutdownModule()
@@ -88,6 +91,8 @@ void FSourceControlModule::ShutdownModule()
 	// we don't care about modular features any more
 	IModularFeatures::Get().OnModularFeatureRegistered().RemoveAll(this);
 	IModularFeatures::Get().OnModularFeatureUnregistered().RemoveAll(this);
+
+	SourceControlFileStatusMonitor.Reset();
 }
 
 void FSourceControlModule::SaveSettings()
@@ -546,6 +551,11 @@ FString FSourceControlModule::GetSourceControlProjectDir() const
 	return FPaths::ProjectDir();
 }
 
+FSourceControlFileStatusMonitor& FSourceControlModule::GetSourceControlFileStatusMonitor()
+{
+	return *SourceControlFileStatusMonitor;
+}
+
 IMPLEMENT_MODULE( FSourceControlModule, SourceControl );
 
 #undef LOCTEXT_NAMESPACE
diff --git a/Engine/Source/Developer/SourceControl/Private/SourceControlModule.h b/Engine/Source/Developer/SourceControl/Private/SourceControlModule.h
--- a/Engine/Source/Developer/SourceControl/Private/SourceControlModule.h
+++ b/Engine/Source/Developer/SourceControl/Private/SourceControlModule.h
@@ -13,6 +13,7 @@
 #include "SourceControlSettings.h"
 #include "DefaultSourceControlProvider.h"
 
+class FSourceControlFileStatusMonitor;
 class SSourceControlLogin;
 class SWindow;
 
@@ -57,6 +58,7 @@ public:
 	virtual void RegisterSourceControlProjectDirDelegate(const FSourceControlProjectDirDelegate& SourceControlProjectDirDelegate) override;
 	virtual void UnregisterSourceControlProjectDirDelegate() override;
 	virtual FString GetSourceControlProjectDir() const override;
+	virtual FSourceControlFileStatusMonitor& GetSourceControlFileStatusMonitor() override;
 
 	/** Save the settings to the ini file */
 	void SaveSettings();
@@ -126,6 +128,9 @@ private:
 	/** The login window control we may be using */
 	TSharedPtr<class SSourceControlLogin> SourceControlLoginPtr;
 
+	/** Monitor the source control status of a collection of files. */
+	TSharedPtr<class FSourceControlFileStatusMonitor> SourceControlFileStatusMonitor;
+
 	/** Files pending a status update */
 	TArray<FString> PendingStatusUpdateFiles;
 
diff --git a/Engine/Source/Developer/SourceControl/Public/ISourceControlModule.h b/Engine/Source/Developer/SourceControl/Public/ISourceControlModule.h
--- a/Engine/Source/Developer/SourceControl/Public/ISourceControlModule.h
+++ b/Engine/Source/Developer/SourceControl/Public/ISourceControlModule.h
@@ -9,6 +9,7 @@
 
 class ISourceControlProvider;
 class ISourceControlChangelist;
+class FSourceControlFileStatusMonitor;
 
 class FSourceControlAssetDataCache;
 class FSourceControlInitSettings;
@@ -244,12 +245,16 @@ public:
 	 */
 	virtual void UnregisterSourceControlProjectDirDelegate() = 0;
 
-
 	/**
 	 * Returns the project base directory to be used by the source control
 	 */
 	virtual FString GetSourceControlProjectDir() const = 0;
 
+	/**
+	 * Returns the object used to monitor the source control status of a collection of files.
+	 */
+	virtual FSourceControlFileStatusMonitor& GetSourceControlFileStatusMonitor() = 0;
+
 	/**
 	 * Gets a reference to the source control module instance.
 	 *
diff --git a/Engine/Source/Developer/SourceControl/Public/ISourceControlState.h b/Engine/Source/Developer/SourceControl/Public/ISourceControlState.h
--- a/Engine/Source/Developer/SourceControl/Public/ISourceControlState.h
+++ b/Engine/Source/Developer/SourceControl/Public/ISourceControlState.h
@@ -178,5 +178,19 @@ public:
 
 	/** Get whether this file can be reverted, i.e. its changes are discarded and the file will no longer be checked-out. */
 	virtual bool CanRevert() const = 0;
+
+	/** Gets the warnings messages associated with this state, if any. Ex 'checkout by other', 'out of date', etc). */
+	virtual TOptional<FText> GetWarningText() const
+	{ 
+		TOptional<FText> WarningText;
+		if (IsConflicted() ||
+			!IsCurrent() ||
+			IsCheckedOutOther() ||
+			(!IsCheckedOut() && (IsCheckedOutInOtherBranch() || IsModifiedInOtherBranch())))
+		{
+			WarningText.Emplace(GetDisplayTooltip()); // The tooltip text usually describe well enough the warning.
+		}
+		return WarningText;
+	}
 };
 
diff --git a/Engine/Source/Developer/SourceControl/Public/SourceControlFileStatusMonitor.h b/Engine/Source/Developer/SourceControl/Public/SourceControlFileStatusMonitor.h
new file mode 100644
--- /dev/null
+++ b/Engine/Source/Developer/SourceControl/Public/SourceControlFileStatusMonitor.h
@@ -0,0 +1,184 @@
+// Copyright Epic Games, Inc. All Rights Reserved.
+
+#pragma once
+
+#include "CoreMinimal.h"
+#include "Containers/Ticker.h"
+#include "Delegates/DelegateCombinations.h"
+#include "ISourceControlOperation.h"
+#include "ISourceControlProvider.h"
+#include "ISourceControlState.h"
+#include "Misc/Timespan.h"
+#include "Templates/SharedPointer.h"
+
+class ISourceControlState;
+
+/**
+ * Periodically monitors the source controlled status of a collection of files and notifies the latest status, including
+ * potential warnings such as 'out of date', 'locked by other', etc. This service is meant to throttle and batch non-time
+ * sensitive FUpdateStatus requests to the source control provider because those requests are expensive and can add an
+ * undesired load on the provider server.
+ */
+class SOURCECONTROL_API FSourceControlFileStatusMonitor : public TSharedFromThis<FSourceControlFileStatusMonitor>
+{
+public:
+	/** Invoked when the status of a source controlled files is updated. The state can be null when the source control provider is changed. */
+	DECLARE_DELEGATE_TwoParams(FOnSourceControlFileStatus, const FString& /*Pathname*/, const ISourceControlState* /*State*/);
+
+public:
+	FSourceControlFileStatusMonitor();
+	virtual ~FSourceControlFileStatusMonitor();
+
+	/**
+	 * Starts monitoring the source control status of the specified file. If the status is already known (the file is already monitored
+	 * by another party), the know status is returned right away, otherwise a request is sent to the source control provider to get the
+	 * status. The request is subject to the probation period policy to allow batching requests. The file status is refreshed according to
+	 * the refresh period policy.
+	 *
+	 * The caller must be sure to call StopMonitoringFile() when the files status is not desired anymore. It is ok to call StopMonitoringFile()
+	 * within the callback OnSourceControlledFileStatus().
+	 *
+	 * @param Owner The unique Id of the caller, typically the caller memory address.
+	 * @param Pathname The absolute path and filname of the file to monitor.
+	 * @param OnSourceControlledFileStatus Delegate invoked whe the status of the file is updated.
+	 *
+	 * @see SetNewRequestProbationPeriod()
+	 * @see SetUpdateStatusPeriod()
+	 * @see SetSuspendMonitoringPolicy()
+	 * @see GetStatusAge()
+	 */
+	void StartMonitoringFile(uintptr_t OwnerId, const FString& Pathname, FOnSourceControlFileStatus OnSourceControlledFileStatus);
+	void StartMonitoringFiles(uintptr_t OwnerId, const TArray<FString>& Pathnames, FOnSourceControlFileStatus OnSourceControlledFileStatus);
+	void StartMonitoringFiles(uintptr_t OwnerId, const TSet<FString>& Pathnames, FOnSourceControlFileStatus OnSourceControlledFileStatus);
+
+	/**
+	 * Stops monitoring the source control status of the specified file. If the specified file is not monitored, the function
+	 * returns successfully.
+	 *
+	 * @param OwnerId The unique Id of the caller, typically the caller memory address.
+	 * @param Pathname The absolute path and filname of the file to stop monitoring.
+	 *
+	 * @note This can be called fron the FOnSourceControlFileStatus callback passed to StartMonitoringFile().
+	 */
+	void StopMonitoringFile(uintptr_t OwnerId, const FString& Pathname);
+	void StopMonitoringFiles(uintptr_t OwnerId, const TArray<FString>& Pathnames);
+	void StopMonitoringFiles(uintptr_t OwnerId, const TSet<FString>& Pathnames);
+
+	/**
+	 * Stops monitoring all the files that were registered by the specified owner Id.
+	 */
+	void StopMonitoringFiles(uintptr_t OwnerId);
+
+	/**
+	 * Returns the set of files being monitored by the specifed owner.
+	 */
+	TSet<FString> GetMonitoredFiles(uintptr_t OwnerId);
+
+	/**
+	 * Returns the file status age, i.e. (now - last_status_update) if the file is
+	 * monitored. If the initial status request hasn't been received yet, returns
+	 * zero.
+	 */
+	TOptional<FTimespan> GetStatusAge(const FString& Pathname) const;
+
+	/**
+	 * Set the times waited before a newly added file status request is sent to the source control provider. This is used batch several
+	 * status for file added closed in time into a single server request, for example if StartMonitoringFile() is called in a loop or
+	 * when the user scroll a long list of files for which the status needs to be updated. Setting a value of 0 remove the probation period.
+	 * By default, the propbation period is 1 second.
+	 */
+	void SetNewRequestProbationPeriodPolicy(const FTimespan& InPropbationPeriod) { ProbationPeriodPolicy = InPropbationPeriod; }
+	FTimespan GetNewRequestProbationPeriodPolicy() const { return ProbationPeriodPolicy; }
+
+	/**
+	 * Set the status refresh period. By default, the monitored files are refreshed every 5 minutes, providing the monitor
+	 * is it suspended according to the suspend monitor policy. The refresh period has an impact on the load of the source control
+	 * provider. A short period may overload the provider service.
+	 * 
+	 * @see SetSuspendMonitoringPolicy
+	 */
+	void SetUpdateStatusPeriodPolicy(const FTimespan& InRefreshPeriod) { RefreshPeriodPolicy = InRefreshPeriod; }
+	FTimespan GetUpdateStatusPeriodPolicy() const { return RefreshPeriodPolicy; }
+
+	/**
+	 * Sets the maximum number of files to batch into a single request. By default, a maximum of 100 files status
+	 * is requested.
+	 */
+	void SetMaxFilePerRequestPolicy(int32 MaxNum) { MaxFileNumPerRequestPolicy = MaxNum; }
+	int32 GetMaxFilePerRequestPolicy() const { return MaxFileNumPerRequestPolicy; }
+
+	/**
+	 * Sets a function invoked to detect if the monitor should suspend its activities, for example, if no user
+	 * is currently interacting with the application, it may not worth checking the file status periodically.
+	 * By default, the monitor will suspends  if there are not Slate interaction within 5 minutes and resume
+	 * at the next user interaction.
+	 */
+	void SetSuspendMonitoringPolicy(TFunction<bool()> IsMonitoringSuspended) { SuspendMonitoringPolicy = MoveTemp(IsMonitoringSuspended); }
+
+private:
+	struct FSourceControlFileStatus
+	{
+		FSourceControlFileStatus(uintptr_t OriginalOwnerId, FOnSourceControlFileStatus InFileStatusUpdateDelegate)
+		{
+			OwnerDelegateMap.Emplace(OriginalOwnerId, MoveTemp(InFileStatusUpdateDelegate));
+		}
+
+		TSharedPtr<ISourceControlState> FileState;
+		double LastStatusCheckTimestampSecs = 0.0;
+		TMap<uintptr_t, FOnSourceControlFileStatus> OwnerDelegateMap;
+	};
+
+private:
+	/** Invoked when the source control has updated the status for the requested files. */
+	void OnSourceControlStatusUpdate(const TSharedRef<ISourceControlOperation>& InOperation, ECommandResult::Type InType);
+
+	/** Invoked when the source control provider changes. */
+	void OnSourceControlProviderChanged(ISourceControlProvider& OldProvider, ISourceControlProvider& NewProvider);
+
+	/** Invoked by the main loop ticker to fire source control requests. */
+	bool Tick(float DeltaTime);
+
+	/** Returns the specified files status or null if not found. */
+	TSharedPtr<FSourceControlFileStatus> FindFileStatus(const FString& Pathname) const;
+
+	/** Whethers the monitor has an in-flight request with the source control provider. */
+	bool HasOngoingRequest() const { return RequestedStatusFiles.Num() > 0; }
+
+private:
+	/** The files currently monitored. */
+	TMap<FString, TSharedPtr<FSourceControlFileStatus>> MonitoredFiles;
+
+	/** The list of files for which a request to the source control provider is currently in-flight/ongoing. */
+	TArray<FString> RequestedStatusFiles;
+
+	/** The newly added files probation period. Time left to batch several files into the same requests. */
+	FTimespan ProbationPeriodPolicy;
+
+	/** The status refresh period. Time left until we re-request a file status. */
+	FTimespan RefreshPeriodPolicy;
+
+	/** Invoked to detect wheter the monitoring is paused/suspended (no request to source control provider).*/
+	TFunction<bool()> SuspendMonitoringPolicy;
+
+	/** The maximum number of files allowed to be batched together for a request to the provider. */
+	int32 MaxFileNumPerRequestPolicy = 100;
+
+	/** Last time a request to the source control provider was issued. */
+	double LastSourceControlCheckSecs = 0.0;
+
+	/** Last time a file was assed to the list of monitored files. */
+	double LastAddedFileTimeSecs = 0.0;
+
+	/** The oldest status update time. */
+	double OldestFileStatusTimeSecs = 0.0;
+
+	/** Number of files newly added that are in probation. */
+	int32 NewAddedFileCount = 0;
+
+	/** Handle for the tick. */
+	FTSTicker::FDelegateHandle TickerHandle;
+
+	/** Handle of the registered source provider change delegate. */
+	FDelegateHandle SourceControlProviderChangedDelegateHandle;
+};
+
diff --git a/Engine/Source/Developer/UnsavedAssetsTracker/Source/Private/SUnsavedAssetsStatusBarWidget.cpp b/Engine/Source/Developer/UnsavedAssetsTracker/Source/Private/SUnsavedAssetsStatusBarWidget.cpp
new file mode 100644
--- /dev/null
+++ b/Engine/Source/Developer/UnsavedAssetsTracker/Source/Private/SUnsavedAssetsStatusBarWidget.cpp
@@ -0,0 +1,89 @@
+// Copyright Epic Games, Inc. All Rights Reserved.
+
+#include "SUnsavedAssetsStatusBarWidget.h"
+
+#include "Widgets/SBoxPanel.h"
+#include "Widgets/Images/SImage.h"
+#include "Widgets/Input/SButton.h"
+#include "Widgets/Text/STextBlock.h"
+#include "Styling/AppStyle.h"
+#include "UnsavedAssetsTracker.h"
+
+#define LOCTEXT_NAMESPACE "UnsavedAssetsTracker"
+
+void SUnsavedAssetsStatusBarWidget::Construct(const FArguments& InArgs, TSharedPtr<FUnsavedAssetsTracker> InUnsavedAssetTracker)
+{
+	UnsavedAssetTracker = MoveTemp(InUnsavedAssetTracker);
+
+	ChildSlot
+	[
+		SNew(SButton)
+		.ContentPadding(FMargin(6.0f,0.0f))
+		.ToolTipText_Lambda([this]() { return GetStatusBarTooltip(); })
+		.ButtonStyle(FAppStyle::Get(), "SimpleButton")
+		.OnClicked(InArgs._OnClicked)
+		[
+			SNew(SHorizontalBox)
+			+ SHorizontalBox::Slot()
+			.AutoWidth()
+			.VAlign(VAlign_Center)
+			.HAlign(HAlign_Center)
+			[
+				SNew(SImage)
+				.Image_Lambda([this](){ return GetStatusBarIcon(); })
+			]
+			+ SHorizontalBox::Slot()
+			.AutoWidth()
+			.VAlign(VAlign_Center)
+			.Padding(FMargin(5, 0, 0, 0))
+			[
+				SNew(STextBlock)
+				.TextStyle(&FAppStyle::Get().GetWidgetStyle<FTextBlockStyle>("NormalText"))
+				.Text_Lambda([this]() { return GetStatusBarText(); })
+			]
+		]
+	];
+}
+
+const FSlateBrush* SUnsavedAssetsStatusBarWidget::GetStatusBarIcon() const
+{
+	if (UnsavedAssetTracker->GetUnsavedAssetNum() == 0)
+	{
+		return FAppStyle::GetBrush("Assets.AllSaved");
+	}
+	else if (UnsavedAssetTracker->GetWarningNum() > 0)
+	{
+		return FAppStyle::GetBrush("Assets.UnsavedWarning");
+	}
+	else
+	{
+		return FAppStyle::GetBrush("Assets.Unsaved");
+	}
+}
+
+FText SUnsavedAssetsStatusBarWidget::GetStatusBarText() const
+{
+	if (UnsavedAssetTracker->GetUnsavedAssetNum() == 0)
+	{
+		return LOCTEXT("All_Saved", "All Saved");
+	}
+	return FText::Format(LOCTEXT("Unsaved Assets", "{0} Unsaved"), UnsavedAssetTracker->GetUnsavedAssetNum());
+}
+
+FText SUnsavedAssetsStatusBarWidget::GetStatusBarTooltip() const
+{
+	if (UnsavedAssetTracker->GetUnsavedAssetNum() == 0)
+	{
+		return LOCTEXT("Assets_All_Saved_Tooltip", "All assets in the project have been saved.");
+	}
+	else if (UnsavedAssetTracker->GetWarningNum() > 0)
+	{
+		return LOCTEXT("Asset_Unsaved_Warning_Tooltip", "Warning: There are currently unsaved assets that have a conflict in Source Control");
+	}
+	else
+	{
+		return FText::Format(LOCTEXT("Asset_Unsaved_Tooltip", "There are currently {0} assets that are unsaved in this project. Press to initiate save."), UnsavedAssetTracker->GetUnsavedAssetNum());
+	}
+}
+
+#undef LOCTEXT_NAMESPACE
diff --git a/Engine/Source/Developer/UnsavedAssetsTracker/Source/Private/SUnsavedAssetsStatusBarWidget.h b/Engine/Source/Developer/UnsavedAssetsTracker/Source/Private/SUnsavedAssetsStatusBarWidget.h
new file mode 100644
--- /dev/null
+++ b/Engine/Source/Developer/UnsavedAssetsTracker/Source/Private/SUnsavedAssetsStatusBarWidget.h
@@ -0,0 +1,36 @@
+// Copyright Epic Games, Inc. All Rights Reserved.
+
+#pragma once
+
+#include "CoreMinimal.h"
+#include "Framework/SlateDelegates.h"
+#include "Templates/SharedPointer.h"
+#include "Widgets/SCompoundWidget.h"
+
+class FUnsavedAssetsTracker;
+
+/**
+ * Tracks assets that has in-memory modification not saved to disk yet and checks
+ * the source control states of those assets when a source control provider is available.
+ */
+class SUnsavedAssetsStatusBarWidget : public SCompoundWidget
+{
+public:
+	SLATE_BEGIN_ARGS(SUnsavedAssetsStatusBarWidget)
+	{}
+		/** Event fired when the status bar button is clicked */
+		SLATE_EVENT(FOnClicked, OnClicked)
+	SLATE_END_ARGS()
+
+	/** Constructs the widget */
+	void Construct(const FArguments& InArgs, TSharedPtr<FUnsavedAssetsTracker> InUnsavedAssetTracker);
+
+private:
+	const FSlateBrush* GetStatusBarIcon() const;
+	FText GetStatusBarText() const;
+	FText GetStatusBarTooltip() const;
+
+private:
+	TSharedPtr<FUnsavedAssetsTracker> UnsavedAssetTracker;
+};
+
diff --git a/Engine/Source/Developer/UnsavedAssetsTracker/Source/Private/UnsavedAssetsTracker.cpp b/Engine/Source/Developer/UnsavedAssetsTracker/Source/Private/UnsavedAssetsTracker.cpp
new file mode 100644
--- /dev/null
+++ b/Engine/Source/Developer/UnsavedAssetsTracker/Source/Private/UnsavedAssetsTracker.cpp
@@ -0,0 +1,344 @@
+// Copyright Epic Games, Inc. All Rights Reserved.
+
+#include "UnsavedAssetsTracker.h"
+
+#include "UnsavedAssetsTrackerModule.h"
+#include "AssetToolsModule.h"
+#include "FileHelpers.h"
+#include "Logging/LogMacros.h"
+#include "ISourceControlModule.h"
+#include "SourceControlHelpers.h"
+#include "SourceControlFileStatusMonitor.h"
+#include "Framework/Notifications/NotificationManager.h"
+#include "Widgets/Notifications/SNotificationList.h"
+#include "Editor/UnrealEdEngine.h"
+#include "UnrealEdGlobals.h"
+
+#define LOCTEXT_NAMESPACE "UnsavedAssetsTracker"
+
+DEFINE_LOG_CATEGORY_STATIC(LogUnsavedAssetsTracker, Log, All);
+
+namespace
+{
+
+bool IsPackagePersistent(UPackage* Package)
+{
+	// Ignore the packages that cannot be saved to disk.
+	if (Package->HasAnyFlags(RF_Transient) || Package->HasAnyPackageFlags(PKG_CompiledIn) || Package == GetTransientPackage())
+	{
+		return false;
+	}
+	return true;
+}
+
+FString GetHumanFriendlyAssetName(UPackage* Package)
+{
+	FName AssetName;
+	FName OwnerName;
+
+	FAssetToolsModule& AssetToolsModule = FModuleManager::LoadModuleChecked<FAssetToolsModule>(TEXT("AssetTools"));
+
+	// Lookup for the first asset in the package
+	UObject* FoundAsset = nullptr;
+	ForEachObjectWithPackage(Package, [&FoundAsset](UObject* InnerObject)
+	{
+		if (InnerObject->IsAsset())
+		{
+			if (FAssetData::IsUAsset(InnerObject))
+			{
+				// If we found the primary asset, use it
+				FoundAsset = InnerObject;
+				return false;
+			}
+			// Otherwise, keep the first found asset but keep looking for a primary asset
+			if (!FoundAsset)
+			{
+				FoundAsset = InnerObject;
+			}
+		}
+		return true;
+	}, /*bIncludeNestedObjects*/ false);
+
+	if (FoundAsset)
+	{
+		TWeakPtr<IAssetTypeActions> AssetTypeActions = AssetToolsModule.Get().GetAssetTypeActionsForClass(FoundAsset->GetClass());
+		if (AssetTypeActions.IsValid())
+		{
+			AssetName = *AssetTypeActions.Pin()->GetObjectDisplayName(FoundAsset);
+		}
+		else
+		{
+			AssetName = FoundAsset->GetFName();
+		}
+
+		OwnerName = FoundAsset->GetOutermostObject()->GetFName();
+	}
+
+	// Last resort, display the package name
+	if (AssetName == NAME_None)
+	{
+		AssetName = *FPackageName::GetShortName(Package->GetFName());
+	}
+
+	return AssetName.ToString();
+}
+
+bool HasPackageWritePermissions(UPackage* Package)
+{
+	// if we do not have write permission under the mount point for this package log an error in the message log to link to.
+	FString PackageName = Package->GetName();
+	return GUnrealEd->HasMountWritePermissionForPackage(PackageName);
+}
+
+} // Anonoymous namespace
+
+
+FUnsavedAssetsTracker::FUnsavedAssetsTracker()
+{
+	// Register for the package dirty state updated callback to catch packages that have been cleaned without being saved
+	UPackage::PackageDirtyStateChangedEvent.AddRaw(this, &FUnsavedAssetsTracker::OnPackageDirtyStateUpdated);
+
+	// Register for the "MarkPackageDirty" callback to catch packages that have been modified and need to be saved
+	UPackage::PackageMarkedDirtyEvent.AddRaw(this, &FUnsavedAssetsTracker::OnPackageMarkedDirty);
+
+	// Register for the package modified callback to catch packages that have been saved
+	UPackage::PackageSavedWithContextEvent.AddRaw(this, &FUnsavedAssetsTracker::OnPackageSaved);
+}
+
+FUnsavedAssetsTracker::~FUnsavedAssetsTracker()
+{
+	UPackage::PackageDirtyStateChangedEvent.RemoveAll(this);
+	UPackage::PackageMarkedDirtyEvent.RemoveAll(this);
+	UPackage::PackageSavedWithContextEvent.RemoveAll(this);
+}
+
+int32 FUnsavedAssetsTracker::GetUnsavedAssetNum() const
+{
+	return UnsavedFiles.Num();
+}
+
+TArray<FString> FUnsavedAssetsTracker::GetUnsavedAssets() const
+{
+	TArray<FString> Pathnames;
+	UnsavedFiles.GetKeys(Pathnames);
+	return Pathnames;
+}
+
+int32 FUnsavedAssetsTracker::GetWarningNum() const
+{
+	return WarningFiles.Num();
+}
+
+TMap<FString, FString> FUnsavedAssetsTracker::GetWarnings() const
+{
+	TMap<FString, FString> Warnings;
+	for (const FString& Pathname: WarningFiles)
+	{
+		if (const FStatus* Status = UnsavedFiles.Find(Pathname))
+		{
+			Warnings.Add(Pathname, Status->WarningText.ToString());
+		}
+		else
+		{
+			checkNoEntry(); // The lists are out of sync.
+		}
+	}
+
+	return Warnings;
+}
+
+void FUnsavedAssetsTracker::OnPackageMarkedDirty(UPackage* Package, bool bWasDirty)
+{
+	AddPackage(Package);
+}
+
+void FUnsavedAssetsTracker::OnPackageDirtyStateUpdated(UPackage* Package)
+{
+	if (Package->IsDirty())
+	{
+		AddPackage(Package);
+	}
+	else
+	{
+		RemovePackage(Package);
+	}
+}
+
+void FUnsavedAssetsTracker::OnPackageSaved(const FString& PackagePathname, UPackage* Package, FObjectPostSaveContext ObjectSaveContext)
+{
+	RemovePackage(Package);
+}
+
+void FUnsavedAssetsTracker::AddPackage(UPackage* Package)
+{
+	if (!IsPackagePersistent(Package))
+	{
+		return;
+	}
+
+	FString PackagePathname = SourceControlHelpers::PackageFilename(Package);
+	if (PackagePathname.IsEmpty())
+	{
+		return;
+	}
+
+	FString HumanFriendlyAssetName = GetHumanFriendlyAssetName(Package);
+
+	int32 UnsavedNumBefore = UnsavedFiles.Num();
+	FStatus& Status = UnsavedFiles.FindOrAdd(PackagePathname, FStatus(HumanFriendlyAssetName));
+	if (UnsavedFiles.Num() > UnsavedNumBefore) // Detect if a new asset was added.
+	{
+		if (!HasPackageWritePermissions(Package))
+		{
+			Status.WarningType = EWarningTypes::PackageWritePermission;
+			Status.WarningText = FText::Format(LOCTEXT("Write_Permission_Warning", "Insufficient writing permission to save {0}"), FText::FromString(Package->GetName()));
+			WarningFiles.Add(PackagePathname);
+			ShowWarningNotificationIfNotAlreadyShown(EWarningTypes::PackageWritePermission, Status.WarningText);
+		}
+
+		ISourceControlModule::Get().GetSourceControlFileStatusMonitor().StartMonitoringFile(
+			reinterpret_cast<uintptr_t>(this),
+			PackagePathname,
+			FSourceControlFileStatusMonitor::FOnSourceControlFileStatus::CreateSP(this, &FUnsavedAssetsTracker::OnSourceControlFileStatusUpdate));
+
+		FUnsavedAssetsTrackerModule::Get().OnUnsavedAssetAdded.Broadcast(PackagePathname);
+
+		UE_LOG(LogUnsavedAssetsTracker, Verbose, TEXT("Added file to the unsaved asset list: %s (%s)"), *HumanFriendlyAssetName, *PackagePathname);
+	}
+}
+
+void FUnsavedAssetsTracker::RemovePackage(UPackage* Package)
+{
+	FString PackagePathname = SourceControlHelpers::PackageFilename(Package);
+	if (PackagePathname.IsEmpty())
+	{
+		return;
+	}
+
+	if (FStatus* Status = UnsavedFiles.Find(PackagePathname))
+	{
+		FString HumanFriendlyAssetName = MoveTemp(Status->HumanFriendlyAssetName); // Keep it for logging.
+
+		UnsavedFiles.Remove(PackagePathname);
+		ISourceControlModule::Get().GetSourceControlFileStatusMonitor().StopMonitoringFile(reinterpret_cast<uintptr_t>(this), PackagePathname);
+
+		// Remove warnings this asset was generating (if any).
+		WarningFiles.Remove(PackagePathname);
+		FUnsavedAssetsTrackerModule::Get().OnUnsavedAssetRemoved.Broadcast(PackagePathname);
+
+		UE_LOG(LogUnsavedAssetsTracker, Log, TEXT("Removed file from the unsaved asset list: %s (%s)"), *HumanFriendlyAssetName, *PackagePathname);
+	}
+}
+
+void FUnsavedAssetsTracker::OnSourceControlFileStatusUpdate(const FString& Pathname, const ISourceControlState* State)
+{
+	auto DiscardWarning = [this](FStatus* InStatus, const FString& Pathname)
+	{
+		// Source control status update cannot clear the package write permission warning.
+		if (InStatus->WarningType != EWarningTypes::PackageWritePermission)
+		{
+			WarningFiles.Remove(Pathname);
+			InStatus->WarningText = FText::GetEmpty();
+			InStatus->WarningType = EWarningTypes::None;
+		}
+
+		if (WarningFiles.IsEmpty()) // All warning were cleared.
+		{
+			ShownWarnings.Reset(); // Reeactivate the notification next time a warning happens.
+		}
+	};
+
+	if (FStatus* Status = UnsavedFiles.Find(Pathname))
+	{
+		if (Status->WarningType == EWarningTypes::PackageWritePermission)
+		{
+			return; // Write permission issue has more weight than source control issue.
+		}
+		else if (State == nullptr) // Source control state was reset. (Changing source control provider/disabling source control)
+		{
+			DiscardWarning(Status, Pathname);
+		}
+		else if (TOptional<FText> WarningText = State->GetWarningText())
+		{
+			Status->WarningText = *WarningText;
+			OnSourceControlWarningNotification(*State, *Status);
+			WarningFiles.Add(Pathname);
+		}
+		else
+		{
+			DiscardWarning(Status, Pathname);
+		}
+	}
+}
+
+void FUnsavedAssetsTracker::PrompToSavePackages()
+{
+	if (GetUnsavedAssetNum() > 0)
+	{
+		const bool bPromptUserToSave = true;
+		const bool bSaveMapPackages = true;
+		const bool bSaveContentPackages = true;
+		const bool bFastSave = false;
+		const bool bClosingEditor = false;
+		const bool bNotifyNoPackagesSaved = true;
+		const bool bCanBeDeclined = false;
+		if (FEditorFileUtils::SaveDirtyPackages(bPromptUserToSave, bSaveMapPackages, bSaveContentPackages, bFastSave, bNotifyNoPackagesSaved, bCanBeDeclined))
+		{
+			// User likely saved something, reset the warnings. We could scan the list of unsaved asset that weren't saved (if any) and check if some warning
+			// types remain, but that looks overkill in this context.
+			ShownWarnings.Reset();
+		}
+	}
+}
+
+void FUnsavedAssetsTracker::OnSourceControlWarningNotification(const ISourceControlState& State, FStatus& InOutStatus)
+{
+	auto UpdateAndShowWarningIfNotAlreadyShown = [this](EWarningTypes WarningType, const FText& Msg, FStatus& InOutStatus)
+	{
+		// Update the warning type.
+		InOutStatus.WarningType = WarningType;
+		ShowWarningNotificationIfNotAlreadyShown(WarningType, Msg);
+	};
+
+	if (State.IsConflicted())
+	{
+		UpdateAndShowWarningIfNotAlreadyShown(EWarningTypes::Conflicted, LOCTEXT("Conflicted_Warning", "Warning: Assets you have edited have conflict(s)."), InOutStatus);
+	}
+	else if (!State.IsCurrent())
+	{
+		UpdateAndShowWarningIfNotAlreadyShown(EWarningTypes::OutOfDate, LOCTEXT("Out_of_Date_Warning", "Warning: Assets you have edited are out of date."), InOutStatus);
+	}
+	else if (State.IsCheckedOutOther())
+	{
+		UpdateAndShowWarningIfNotAlreadyShown(EWarningTypes::CheckedOutByOther, LOCTEXT("Locked_by_Other_Warning", "Warning: Assets you have edited are locked by another user."), InOutStatus);
+	}
+	else if (!State.IsCheckedOut())
+	{
+		if (State.IsCheckedOutInOtherBranch())
+		{
+			UpdateAndShowWarningIfNotAlreadyShown(EWarningTypes::CheckedOutInOtherBranch, LOCTEXT("Checked_Out_In_Other_Branch_Warning", "Warning: Assets you have edited are checked out in another branch."), InOutStatus);
+		}
+		else if (State.IsModifiedInOtherBranch())
+		{
+			UpdateAndShowWarningIfNotAlreadyShown(EWarningTypes::ModifiedInOtherBranch, LOCTEXT("Modified_In_Other_Branch_Warning", "Warning: Assets you have edited are modified in another branch."), InOutStatus);
+		}
+	}
+	else if (State.GetWarningText().IsSet())
+	{
+		UpdateAndShowWarningIfNotAlreadyShown(EWarningTypes::Other, LOCTEXT("Generic_Warning", "Warning: Assets you have edited have warnings."), InOutStatus);
+	}
+}
+
+void FUnsavedAssetsTracker::ShowWarningNotificationIfNotAlreadyShown(EWarningTypes WarningType, const FText& Msg)
+{
+	// Show the notification if it hasn't been shown since the last reset/save.
+	if (bWarningNotificationEnabled && !ShownWarnings.Contains(WarningType))
+	{
+		// Setup the notification for operation feedback
+		FNotificationInfo Info(Msg);
+		Info.Image = FAppStyle::GetBrush("Icons.WarningWithColor");
+		TSharedPtr<SNotificationItem> Notification = FSlateNotificationManager::Get().AddNotification(Info);
+		Notification->SetCompletionState(SNotificationItem::CS_None);
+		ShownWarnings.Add(WarningType);
+	}
+}
diff --git a/Engine/Source/Developer/UnsavedAssetsTracker/Source/Private/UnsavedAssetsTracker.h b/Engine/Source/Developer/UnsavedAssetsTracker/Source/Private/UnsavedAssetsTracker.h
new file mode 100644
--- /dev/null
+++ b/Engine/Source/Developer/UnsavedAssetsTracker/Source/Private/UnsavedAssetsTracker.h
@@ -0,0 +1,103 @@
+// Copyright Epic Games, Inc. All Rights Reserved.
+
+#pragma once
+
+#include "CoreMinimal.h"
+#include "Templates/SharedPointer.h"
+#include "UObject/ObjectSaveContext.h"
+
+class UPackage;
+class ISourceControlState;
+
+/**
+ * Tracks assets that has in-memory modification not saved to disk yet and checks
+ * the source control states of those assets when a source control provider is available.
+ */
+class FUnsavedAssetsTracker : public TSharedFromThis<FUnsavedAssetsTracker>
+{
+public:
+	FUnsavedAssetsTracker();
+	virtual ~FUnsavedAssetsTracker();
+
+	/** Returns the number of unsaved assets currently tracked. */
+	int32 GetUnsavedAssetNum() const;
+
+	/** Returns the list of unsaved assets. */
+	TArray<FString> GetUnsavedAssets() const;
+
+	/** Returns the number of active warnings. */
+	int32 GetWarningNum() const;
+
+	/** Returns a map of pathname/warnings. */
+	TMap<FString, FString> GetWarnings() const;
+
+	/** Sets whether a toast notification is emitted when looking up the source control state of an unsaved asset generates a warning. */
+	void SetWarningNotificationEnabled(bool bEnabled) { bWarningNotificationEnabled = bEnabled; }
+
+	/** Displays a dialog prompting the user to save the packages. */
+	void PrompToSavePackages();
+
+private:
+	enum class EWarningTypes
+	{
+		None,
+		Conflicted,
+		OutOfDate,
+		CheckedOutByOther,
+		CheckedOutInOtherBranch,
+		ModifiedInOtherBranch,
+		PackageWritePermission,
+		Other,
+	};
+
+	struct FStatus
+	{
+		FStatus(FString InHumanFriendlyAssetName) : HumanFriendlyAssetName(InHumanFriendlyAssetName) {}
+		FString HumanFriendlyAssetName;
+		EWarningTypes WarningType = EWarningTypes::None;
+		FText WarningText;
+	};
+
+	/**
+	 * Invoked when a package is marked dirty.
+	 *
+	 * @param Pkg The package that was marked as dirty
+	 * @param bWasDirty Whether the package was previously dirty before the call to MarkPackageDirty
+	 */
+	void OnPackageMarkedDirty(UPackage* Package, bool bWasDirty);
+
+	/**
+	 * Invoked when a package is saved.
+	 *
+	 * @param Filename The filename the package was saved to
+	 * @param Package The package that was saved
+	 * @param ObjectSaveContext The package saving context.
+	 */
+	void OnPackageSaved(const FString& PackagePathname, UPackage* Package, FObjectPostSaveContext ObjectSaveContext);
+
+	/**
+	 * Invoked when the dirty state of a package changed. Used to track when a package was cleaned without being saved.
+	 */
+	void OnPackageDirtyStateUpdated(UPackage* Package);
+
+	/** Starts to track an unsaved package.*/
+	void AddPackage(UPackage* Package);
+
+	/** Stops tracking the specified unsaved package. */
+	void RemovePackage(UPackage* Package);
+
+	/** Invoked when the status of a source controlled is updated, with the corresponding warning, if any. */
+	void OnSourceControlFileStatusUpdate(const FString& Pathname, const ISourceControlState* Status);
+
+	/** Invoked when a new warning is discovered. */
+	void OnSourceControlWarningNotification(const ISourceControlState& State, FStatus& InOutStatus);
+
+	/** Show a toast notifications if the warning types hasn't been shown yet. */
+	void ShowWarningNotificationIfNotAlreadyShown(EWarningTypes WarningType, const FText& Msg);
+
+private:
+	TMap<FString, FStatus> UnsavedFiles;
+	TSet<FString> WarningFiles;
+	TSet<EWarningTypes> ShownWarnings;
+	bool bWarningNotificationEnabled = true;
+};
diff --git a/Engine/Source/Developer/UnsavedAssetsTracker/Source/Private/UnsavedAssetsTrackerModule.cpp b/Engine/Source/Developer/UnsavedAssetsTracker/Source/Private/UnsavedAssetsTrackerModule.cpp
new file mode 100644
--- /dev/null
+++ b/Engine/Source/Developer/UnsavedAssetsTracker/Source/Private/UnsavedAssetsTrackerModule.cpp
@@ -0,0 +1,61 @@
+// Copyright Epic Games, Inc. All Rights Reserved.
+
+#include "UnsavedAssetsTrackerModule.h"
+
+#include "CoreGlobals.h"
+#include "Misc/App.h"
+#include "UnsavedAssetsTracker.h"
+#include "SUnsavedAssetsStatusBarWidget.h"
+#include "Widgets/SNullWidget.h"
+
+
+void FUnsavedAssetsTrackerModule::StartupModule()
+{
+	UnsavedAssetTracker = MakeShared<FUnsavedAssetsTracker>();
+}
+
+void FUnsavedAssetsTrackerModule::ShutdownModule()
+{
+	UnsavedAssetTracker.Reset();
+}
+
+int32 FUnsavedAssetsTrackerModule::GetUnsavedAssetNum() const
+{
+	if (UnsavedAssetTracker)
+	{
+		return UnsavedAssetTracker->GetUnsavedAssetNum();
+	}
+	return 0;
+}
+
+TArray<FString> FUnsavedAssetsTrackerModule::GetUnsavedAssets() const
+{
+	if (UnsavedAssetTracker)
+	{
+		return UnsavedAssetTracker->GetUnsavedAssets();
+	}
+	return TArray<FString>();
+}
+
+TSharedRef<SWidget> FUnsavedAssetsTrackerModule::MakeUnsavedAssetsStatusBarWidget()
+{
+	if (!UnsavedAssetTracker)
+	{
+		return SNullWidget::NullWidget;
+	}
+
+	TWeakPtr<FUnsavedAssetsTracker> WeakUnsavedAssetTracker(UnsavedAssetTracker);
+	return SNew(SUnsavedAssetsStatusBarWidget, UnsavedAssetTracker)
+		.OnClicked_Lambda([WeakUnsavedAssetTracker]()
+		{
+			if (TSharedPtr<FUnsavedAssetsTracker> Tracker = WeakUnsavedAssetTracker.Pin())
+			{
+				Tracker->PrompToSavePackages();
+			}
+			return FReply::Handled();
+		});
+}
+
+IMPLEMENT_MODULE(FUnsavedAssetsTrackerModule, UnsavedAssetsTracker)
+
+#undef LOCTEXT_NAMESPACE
diff --git a/Engine/Source/Developer/UnsavedAssetsTracker/Source/Public/UnsavedAssetsTrackerModule.h b/Engine/Source/Developer/UnsavedAssetsTracker/Source/Public/UnsavedAssetsTrackerModule.h
new file mode 100644
--- /dev/null
+++ b/Engine/Source/Developer/UnsavedAssetsTracker/Source/Public/UnsavedAssetsTrackerModule.h
@@ -0,0 +1,59 @@
+// Copyright Epic Games, Inc. All Rights Reserved.
+
+#pragma once
+
+#include "CoreMinimal.h"
+#include "EditorSubsystem.h"
+#include "Modules/ModuleInterface.h"
+#include "Modules/ModuleManager.h"
+#include "Delegates/DelegateCombinations.h"
+#include "Templates/SharedPointer.h"
+#include "UObject/ObjectSaveContext.h"
+
+class FUnsavedAssetsTracker;
+class SWidget;
+
+/** Invoked when an asset was modified in memory and added to the tracker list. */
+DECLARE_MULTICAST_DELEGATE_OneParam(FOnUnsavedAssetAdded, const FString& /*FileAbsPathname*/);
+
+/** Invoked when an asset modified in memory was saved and removed from the tracker list. */
+DECLARE_MULTICAST_DELEGATE_OneParam(FOnUnsavedAssetRemoved, const FString& /*FileAbsPathname*/);
+
+
+/**
+ * Tracks assets that has in-memory modification not saved to disk yet and checks
+ * the source control states of those assets when a source control provider is available.
+ */
+class UNSAVEDASSETSTRACKER_API FUnsavedAssetsTrackerModule : public IModuleInterface
+{
+public:
+	/** IModuleInterface implementation */
+	virtual void StartupModule() override;
+	virtual void ShutdownModule() override;
+
+	/**
+	 * Retrieve the module instance.
+	 */
+	static inline FUnsavedAssetsTrackerModule& Get()
+	{
+		return FModuleManager::LoadModuleChecked<FUnsavedAssetsTrackerModule>("UnsavedAssetsTracker");
+	}
+
+	/** Returns the number of unsaved assets currently tracked. */
+	int32 GetUnsavedAssetNum() const;
+
+	/** Returns the list of unsaved assets. */
+	TArray<FString> GetUnsavedAssets() const;
+
+	/** Construct the widget for the Editor status bar.*/
+	TSharedRef<SWidget> MakeUnsavedAssetsStatusBarWidget();
+
+	/** Invoked when a file is added to the unsaved list. */
+	FOnUnsavedAssetAdded OnUnsavedAssetAdded;
+
+	/** Invoked when a file is removed from the unsaved list. */
+	FOnUnsavedAssetAdded OnUnsavedAssetRemoved;
+
+private:
+	TSharedPtr<FUnsavedAssetsTracker> UnsavedAssetTracker;
+};
diff --git a/Engine/Source/Developer/UnsavedAssetsTracker/Source/UnsavedAssetsTracker.Build.cs b/Engine/Source/Developer/UnsavedAssetsTracker/Source/UnsavedAssetsTracker.Build.cs
new file mode 100644
--- /dev/null
+++ b/Engine/Source/Developer/UnsavedAssetsTracker/Source/UnsavedAssetsTracker.Build.cs
@@ -0,0 +1,30 @@
+// Copyright Epic Games, Inc. All Rights Reserved.
+
+using UnrealBuildTool;
+
+public class UnsavedAssetsTracker : ModuleRules
+{
+	public UnsavedAssetsTracker(ReadOnlyTargetRules Target) : base(Target)
+	{
+		PublicDependencyModuleNames.AddRange(
+			new string[]
+			{
+				"Core",
+				"Engine",
+				"EditorSubsystem",
+				"SourceControl",
+			}
+		);
+
+		PrivateDependencyModuleNames.AddRange(
+			new string[] {
+				"Core",
+				"CoreUObject",
+				"Engine",
+				"SlateCore",
+				"Slate",
+				"UnrealEd",
+			}
+		);
+	}
+}
diff --git a/Engine/Source/Developer/UnsavedAssetsTracker/UnsavedAssetTracker.uplugin b/Engine/Source/Developer/UnsavedAssetsTracker/UnsavedAssetTracker.uplugin
new file mode 100644
--- /dev/null
+++ b/Engine/Source/Developer/UnsavedAssetsTracker/UnsavedAssetTracker.uplugin
@@ -0,0 +1,26 @@
+{
+	"FileVersion": 3,
+	"Version": 1,
+	"VersionName": "1.0",
+	"FriendlyName": "Unsaved Assets Tracker",
+	"Description": "",
+	"Category": "Editor",
+	"CreatedBy": "Epic Games, Inc.",
+	"CreatedByURL": "http://epicgames.com",
+	"DocsURL": "",
+	"MarketplaceURL": "",
+	"SupportURL": "",
+	"EnabledByDefault": true,
+	"CanContainContent": false,
+	"IsBetaVersion": false,
+	"Installed": false,
+	"Modules": [
+		{
+			"Name": "UnsavedAssetsTracker",
+			"Type": "Editor",
+			"LoadingPhase": "Default"
+		}
+	],
+	"Plugins": [
+	]
+}
diff --git a/Engine/Source/Editor/EditorStyle/Private/StarshipStyle.cpp b/Engine/Source/Editor/EditorStyle/Private/StarshipStyle.cpp
--- a/Engine/Source/Editor/EditorStyle/Private/StarshipStyle.cpp
+++ b/Engine/Source/Editor/EditorStyle/Private/StarshipStyle.cpp
@@ -243,6 +243,7 @@ void FStarshipEditorStyle::FStyle::Initialize()
 	SetupToolkitStyles();
 	SetupTranslationEditorStyles();
 	SetupLocalizationDashboardStyles();
+	SetupUnsavedAssetsStyles();
 	SetupSourceControlStyles();
 	SetupAutomationStyles();
 	SetupUMGEditorStyles();
@@ -7029,6 +7030,15 @@ void FStarshipEditorStyle::FStyle::SetupToolkitStyles()
 #endif // WITH_EDITOR
 }
 
+void FStarshipEditorStyle::FStyle::SetupUnsavedAssetsStyles()
+{
+#if WITH_EDITOR || (IS_PROGRAM && WITH_UNREAL_DEVELOPER_TOOLS)
+	Set("Assets.AllSaved", new CORE_IMAGE_BRUSH_SVG("Starship/Common/AllSavedAssets", Icon16x16, FStyleColors::Foreground));
+	Set("Assets.Unsaved", new CORE_IMAGE_BRUSH_SVG("Starship/Common/UnsavedAssets", Icon16x16, FStyleColors::Foreground));
+	Set("Assets.UnsavedWarning", new CORE_IMAGE_BRUSH_SVG("Starship/Common/UnsavedAssetsWarning", Icon16x16, FStyleColors::AccentYellow));
+#endif
+}
+
 void FStarshipEditorStyle::FStyle::SetupSourceControlStyles()
 {
 	//Source Control
diff --git a/Engine/Source/Editor/EditorStyle/Private/StarshipStyle.h b/Engine/Source/Editor/EditorStyle/Private/StarshipStyle.h
--- a/Engine/Source/Editor/EditorStyle/Private/StarshipStyle.h
+++ b/Engine/Source/Editor/EditorStyle/Private/StarshipStyle.h
@@ -68,6 +68,7 @@ public:
 		void SetupContentBrowserStyle();
 		void SetupLandscapeEditorStyle();
 		void SetupToolkitStyles();
+		void SetupUnsavedAssetsStyles();
 		void SetupSourceControlStyles();
 		void SetupAutomationStyles();
 		void SetupUMGEditorStyles();
diff --git a/Engine/Source/Editor/SourceControlWindows/Private/SSourceControlChangelists.cpp b/Engine/Source/Editor/SourceControlWindows/Private/SSourceControlChangelists.cpp
--- a/Engine/Source/Editor/SourceControlWindows/Private/SSourceControlChangelists.cpp
+++ b/Engine/Source/Editor/SourceControlWindows/Private/SSourceControlChangelists.cpp
@@ -26,6 +26,7 @@
 #include "SSourceControlDescription.h"
 #include "SourceControlWindows.h"
 #include "SourceControlHelpers.h"
+#include "SourceControlFileStatusMonitor.h"
 #include "SourceControlPreferences.h"
 #include "SourceControlMenuContext.h"
 #include "SourceControlSettings.h"
@@ -664,7 +665,7 @@ void SSourceControlChangelistsWidget::RequestChangelistsRefresh()
 
 void SSourceControlChangelistsWidget::RequestFileStatusRefresh(const IChangelistTreeItem& Changelist)
 {
-	TArray<FString> Pathnames;
+	TSet<FString> Pathnames;
 	if (Changelist.GetTreeItemType() == IChangelistTreeItem::Changelist)
 	{
 		const FChangelistTreeItem& ChangelistItem = static_cast<const FChangelistTreeItem&>(Changelist);
@@ -689,33 +690,32 @@ void SSourceControlChangelistsWidget::RequestFileStatusRefresh(const IChangelist
 	}
 }
 
-void SSourceControlChangelistsWidget::RequestFileStatusRefresh(const TArray<FString>& Pathnames)
+void SSourceControlChangelistsWidget::RequestFileStatusRefresh(const TSet<FString>& PathnamesToMonitor)
 {
 	if (!ISourceControlModule::Get().IsEnabled())
 	{
 		return;
 	}
 
+	// NOTE: This is using the file status monitor and the status might be a bit old but not older than the
+	//       update period policy of the monitor. See FSourceControlFileStatusMonitor::SetUpdateStatusPeriodPolicy()
+
 	// If the changelist contains files.
-	if (!Pathnames.IsEmpty())
+	if (!PathnamesToMonitor.IsEmpty())
 	{
-		TSharedRef<FUpdateStatus> UpdateStatus = ISourceControlOperation::Create<FUpdateStatus>();
-
-		if (FUncontrolledChangelistsModule::Get().IsEnabled())
-		{
-			// Uncontrolled CL files are not opened in the P4 client view, force the status update of the files even if not opened to capture who has it in checked out.
-			UpdateStatus->SetForceUpdate(true);
-		}
-
-		OnStartSourceControlOperation(UpdateStatus, LOCTEXT("Updating_Changelist_Files_Status", "Updating changelist files status..."));
-
-		// Check if the source controlled states of those files changed outside the Editor. The Editor prevents a user from checking out assets that are out
-		// of date, but the user can checkout out from P4V directly. This request is meant to detect if a users checkout outdated assets outside the Editor. It will
-		// also detect non-exclusive files that are out of date (.ini files for example). That's an edge case and if this creates bottleneck, this can likely be disabled.
-		ISourceControlModule::Get().GetProvider().Execute(UpdateStatus, Pathnames, EConcurrency::Asynchronous, FSourceControlOperationComplete::CreateSP(this, &SSourceControlChangelistsWidget::OnEndSourceControlOperation));
+		// Request an update only for the files that changed.
+		TSet<FString> MonitoredPathnames = ISourceControlModule::Get().GetSourceControlFileStatusMonitor().GetMonitoredFiles(reinterpret_cast<uintptr_t>(this));
+		TSet<FString> StopMonitoringPathnames = MonitoredPathnames.Difference(PathnamesToMonitor);
+		TSet<FString> StartMonitoringPathnames = PathnamesToMonitor.Difference(MonitoredPathnames);
 
+		ISourceControlModule::Get().GetSourceControlFileStatusMonitor().StopMonitoringFiles(reinterpret_cast<uintptr_t>(this), StopMonitoringPathnames);
+		ISourceControlModule::Get().GetSourceControlFileStatusMonitor().StartMonitoringFiles(reinterpret_cast<uintptr_t>(this), StartMonitoringPathnames, FSourceControlFileStatusMonitor::FOnSourceControlFileStatus());
 		// NOTE: The updated status are expected to come from OnSourceControlStateChanged() that is invoked when the source control internal state changes.
 	}
+	else
+	{
+		ISourceControlModule::Get().GetSourceControlFileStatusMonitor().StopMonitoringFiles(reinterpret_cast<uintptr_t>(this));
+	}
 }
 
 void SSourceControlChangelistsWidget::StartRefreshStatus()
diff --git a/Engine/Source/Editor/SourceControlWindows/Private/SSourceControlChangelists.h b/Engine/Source/Editor/SourceControlWindows/Private/SSourceControlChangelists.h
--- a/Engine/Source/Editor/SourceControlWindows/Private/SSourceControlChangelists.h
+++ b/Engine/Source/Editor/SourceControlWindows/Private/SSourceControlChangelists.h
@@ -72,7 +72,7 @@ private:
 
 	void RequestChangelistsRefresh();
 	void RequestFileStatusRefresh(const IChangelistTreeItem& Changelist);
-	void RequestFileStatusRefresh(const TArray<FString>& Pathnames);
+	void RequestFileStatusRefresh(const TSet<FString>& Pathnames);
 	void OnRefresh();
 	void ClearChangelistsTree();
 
diff --git a/Engine/Source/Editor/StatusBar/Private/SourceControlMenuHelpers.cpp b/Engine/Source/Editor/StatusBar/Private/SourceControlMenuHelpers.cpp
--- a/Engine/Source/Editor/StatusBar/Private/SourceControlMenuHelpers.cpp
+++ b/Engine/Source/Editor/StatusBar/Private/SourceControlMenuHelpers.cpp
@@ -8,6 +8,7 @@
 #include "ISourceControlWindowsModule.h"
 #include "SourceControlMenuHelpers.h"
 #include "SourceControlWindows.h"
+#include "UnsavedAssetsTrackerModule.h"
 #include "FileHelpers.h"
 #include "ToolMenuContext.h"
 #include "ToolMenus.h"
@@ -484,7 +485,16 @@ TSharedRef<SWidget> FSourceControlMenuHelpers::MakeSourceControlStatusWidget()
 		.VAlign(VAlign_Center)
 		.AutoWidth()
 		[
-			SNew(SButton)			
+			// NOTE: The unsaved can have its own menu extension in the status bar to decouple it from source control, but putting all the buttons in the right order
+			//       on the status bar is not deterministic, you need to know the name of the menu that is before or after and the menu is dynamic. Having it here
+			//       ensure its position with respect to the source control button.
+			FUnsavedAssetsTrackerModule::Get().MakeUnsavedAssetsStatusBarWidget()
+		]
+		+ SHorizontalBox::Slot()
+		.VAlign(VAlign_Center)
+		.AutoWidth()
+		[
+			SNew(SButton)
 			.ButtonStyle(&FAppStyle::Get().GetWidgetStyle<FButtonStyle>("StatusBar.StatusBarButton"))
 			.ToolTipText_Static(&FSourceControlMenuHelpers::GetSourceControlSyncStatusTooltipText)
 			.Visibility_Static(&FSourceControlMenuHelpers::GetSourceControlSyncStatusVisibility)
diff --git a/Engine/Source/Editor/StatusBar/StatusBar.Build.cs b/Engine/Source/Editor/StatusBar/StatusBar.Build.cs
--- a/Engine/Source/Editor/StatusBar/StatusBar.Build.cs
+++ b/Engine/Source/Editor/StatusBar/StatusBar.Build.cs
@@ -23,6 +23,7 @@ public class StatusBar : ModuleRules
 				"InputCore",
 				"EditorFramework",
 				"SourceControlWindows",
+				"UnsavedAssetsTracker",
 				"ToolMenus",
 				"ToolWidgets",
 				"UnrealEd",
diff --git a/Engine/Source/Editor/UnrealEd/Private/FileHelpers.cpp b/Engine/Source/Editor/UnrealEd/Private/FileHelpers.cpp
--- a/Engine/Source/Editor/UnrealEd/Private/FileHelpers.cpp
+++ b/Engine/Source/Editor/UnrealEd/Private/FileHelpers.cpp
@@ -4262,13 +4262,44 @@ FEditorFileUtils::EPromptReturnCode FEditorFileUtils::PromptForCheckoutAndSave(
 
 		if ( AddPackageItemsUnchecked.Num() > 0 || AddPackageItemsChecked.Num() > 0 )
 		{
+			int32 WarningCount = 0;
+			auto AddPackageItem = [&PackagesDialogModule, &WarningCount](UPackage* Package, ECheckBoxState CheckedState)
+			{
+				FString IconName;
+				FString IconTooltip;
+				if (!GUnrealEd->HasMountWritePermissionForPackage(Package->GetName()))
+				{
+					IconName = TEXT("Icons.WarningWithColor");
+					IconTooltip = TEXT("Insufficient writing permission to save");
+					++WarningCount;
+				}
+				else if (ISourceControlModule::Get().IsEnabled())
+				{
+					if (TSharedPtr<ISourceControlState> State = ISourceControlModule::Get().GetProvider().GetState(Package, EStateCacheUsage::Use))
+					{
+						if (TOptional<FText> Warning = State->GetWarningText())
+						{
+							IconName = TEXT("Icons.WarningWithColor");
+							IconTooltip = Warning->ToString();
+							++WarningCount;
+						}
+					}
+				}
+				PackagesDialogModule.AddPackageItem(Package, CheckedState, /*Disabled*/false, IconName, IconTooltip);
+			};
+
 			for (auto Iter = AddPackageItemsChecked.CreateIterator(); Iter; ++Iter)
 			{
-				PackagesDialogModule.AddPackageItem(*Iter, ECheckBoxState::Checked);
+				AddPackageItem(*Iter, ECheckBoxState::Checked);
 			}
 			for (auto Iter = AddPackageItemsUnchecked.CreateIterator(); Iter; ++Iter)
 			{
-				PackagesDialogModule.AddPackageItem(*Iter, ECheckBoxState::Unchecked);
+				AddPackageItem(*Iter, ECheckBoxState::Unchecked);
+			}
+
+			if (WarningCount > 0)
+			{
+				PackagesDialogModule.SetWarning(LOCTEXT("Warning_Notification", "Warning: Assets have conflict in Source Control or cannot be written to disk"));
 			}
 
 			// If valid packages were added to the dialog, display it to the user
