Fixed a crash in the UnsavedAssetsTracker when packages in the /Memory mount point gets dirty.
Fixed the UnsavedAssetsTracker to not track 'PlayInEditor' packages and packages from /Temp

diff --git a/Engine/Source/Developer/UnsavedAssetsTracker/Source/Private/UnsavedAssetsTracker.cpp b/Engine/Source/Developer/UnsavedAssetsTracker/Source/Private/UnsavedAssetsTracker.cpp
--- a/Engine/Source/Developer/UnsavedAssetsTracker/Source/Private/UnsavedAssetsTracker.cpp
+++ b/Engine/Source/Developer/UnsavedAssetsTracker/Source/Private/UnsavedAssetsTracker.cpp
@@ -13,6 +13,7 @@
 #include "Widgets/Notifications/SNotificationList.h"
 #include "Editor/UnrealEdEngine.h"
 #include "UnrealEdGlobals.h"
+#include "Misc/PackageName.h"
 
 #define LOCTEXT_NAMESPACE "UnsavedAssetsTracker"
 
@@ -23,8 +24,13 @@ namespace
 
 bool IsPackagePersistent(UPackage* Package)
 {
-	// Ignore the packages that cannot be saved to disk.
-	if (Package->HasAnyFlags(RF_Transient) || Package->HasAnyPackageFlags(PKG_CompiledIn) || Package == GetTransientPackage())
+	// Ignore the packages that aren't meant to be persistent.
+	if (Package->HasAnyFlags(RF_Transient) ||
+		Package->HasAnyPackageFlags(PKG_CompiledIn) ||
+		Package->HasAnyPackageFlags(PKG_PlayInEditor) ||
+		Package == GetTransientPackage() ||
+		FPackageName::IsMemoryPackage(Package->GetPathName()) ||
+		FPackageName::IsTempPackage(Package->GetPathName()))
 	{
 		return false;
 	}
