Fixed missing #undef LOCTEXT in UnsavedAssetsTracker.cpp the extra #undef LOCTEXT in UnsavedAssetsTrackerModule.cpp.
Fixed log level in UnsavedAssetsTracker to be verbose rather than log.

diff --git a/Engine/Source/Developer/UnsavedAssetsTracker/Source/Private/UnsavedAssetsTracker.cpp b/Engine/Source/Developer/UnsavedAssetsTracker/Source/Private/UnsavedAssetsTracker.cpp
--- a/Engine/Source/Developer/UnsavedAssetsTracker/Source/Private/UnsavedAssetsTracker.cpp
+++ b/Engine/Source/Developer/UnsavedAssetsTracker/Source/Private/UnsavedAssetsTracker.cpp
@@ -29,8 +29,7 @@ bool IsPackagePersistent(UPackage* Package)
 		Package->HasAnyPackageFlags(PKG_CompiledIn) ||
 		Package->HasAnyPackageFlags(PKG_PlayInEditor) ||
 		Package == GetTransientPackage() ||
-		FPackageName::IsMemoryPackage(Package->GetPathName()) ||
-		FPackageName::IsTempPackage(Package->GetPathName()))
+		FPackageName::IsMemoryPackage(Package->GetPathName()))
 	{
 		return false;
 	}
@@ -215,6 +214,11 @@ void FUnsavedAssetsTracker::AddPackage(UPackage* Package)
 
 void FUnsavedAssetsTracker::RemovePackage(UPackage* Package)
 {
+	if (!IsPackagePersistent(Package))
+	{
+		return;
+	}
+
 	FString PackagePathname = SourceControlHelpers::PackageFilename(Package);
 	if (PackagePathname.IsEmpty())
 	{
@@ -232,7 +236,7 @@ void FUnsavedAssetsTracker::RemovePackage(UPackage* Package)
 		WarningFiles.Remove(PackagePathname);
 		FUnsavedAssetsTrackerModule::Get().OnUnsavedAssetRemoved.Broadcast(PackagePathname);
 
-		UE_LOG(LogUnsavedAssetsTracker, Log, TEXT("Removed file from the unsaved asset list: %s (%s)"), *HumanFriendlyAssetName, *PackagePathname);
+		UE_LOG(LogUnsavedAssetsTracker, Verbose, TEXT("Removed file from the unsaved asset list: %s (%s)"), *HumanFriendlyAssetName, *PackagePathname);
 	}
 }
 
@@ -348,3 +352,5 @@ void FUnsavedAssetsTracker::ShowWarningNotificationIfNotAlreadyShown(EWarningTyp
 		ShownWarnings.Add(WarningType);
 	}
 }
+
+#undef LOCTEXT_NAMESPACE
diff --git a/Engine/Source/Developer/UnsavedAssetsTracker/Source/Private/UnsavedAssetsTrackerModule.cpp b/Engine/Source/Developer/UnsavedAssetsTracker/Source/Private/UnsavedAssetsTrackerModule.cpp
--- a/Engine/Source/Developer/UnsavedAssetsTracker/Source/Private/UnsavedAssetsTrackerModule.cpp
+++ b/Engine/Source/Developer/UnsavedAssetsTracker/Source/Private/UnsavedAssetsTrackerModule.cpp
@@ -58,4 +58,3 @@ TSharedRef<SWidget> FUnsavedAssetsTrackerModule::MakeUnsavedAssetsStatusBarWidge
 
 IMPLEMENT_MODULE(FUnsavedAssetsTrackerModule, UnsavedAssetsTracker)
 
-#undef LOCTEXT_NAMESPACE
