Fixed the unsaved assets tracker crashing when auto-saving package.
Fixed the unsaved assets tracker to correctly work undo/redo.
Fixed the unsaved assets tracker to correctly handle when the temporary 'Untitled' map is discarded or saved.

diff --git a/Engine/Source/Developer/UnsavedAssetsTracker/Source/Private/UnsavedAssetsTracker.cpp b/Engine/Source/Developer/UnsavedAssetsTracker/Source/Private/UnsavedAssetsTracker.cpp
--- a/Engine/Source/Developer/UnsavedAssetsTracker/Source/Private/UnsavedAssetsTracker.cpp
+++ b/Engine/Source/Developer/UnsavedAssetsTracker/Source/Private/UnsavedAssetsTracker.cpp
@@ -11,9 +11,13 @@
 #include "SourceControlFileStatusMonitor.h"
 #include "Framework/Notifications/NotificationManager.h"
 #include "Widgets/Notifications/SNotificationList.h"
+#include "Editor.h"
 #include "Editor/UnrealEdEngine.h"
+#include "Editor/TransBuffer.h"
+#include "Engine/World.h"
 #include "UnrealEdGlobals.h"
 #include "Misc/PackageName.h"
+#include "LevelEditor.h"
 
 #define LOCTEXT_NAMESPACE "UnsavedAssetsTracker"
 
@@ -22,7 +26,7 @@ DEFINE_LOG_CATEGORY_STATIC(LogUnsavedAssetsTracker, Log, All);
 namespace
 {
 
-bool IsPackagePersistent(UPackage* Package)
+bool ShouldTrackDirtyPackage(const UPackage* Package)
 {
 	// Ignore the packages that aren't meant to be persistent.
 	if (Package->HasAnyFlags(RF_Transient) ||
@@ -36,7 +40,7 @@ bool IsPackagePersistent(UPackage* Package)
 	return true;
 }
 
-FString GetHumanFriendlyAssetName(UPackage* Package)
+FString GetHumanFriendlyAssetName(const UPackage* Package)
 {
 	FName AssetName;
 	FName OwnerName;
@@ -88,13 +92,43 @@ FString GetHumanFriendlyAssetName(UPackage* Package)
 	return AssetName.ToString();
 }
 
-bool HasPackageWritePermissions(UPackage* Package)
+bool HasPackageWritePermissions(const UPackage* Package)
 {
 	// if we do not have write permission under the mount point for this package log an error in the message log to link to.
 	FString PackageName = Package->GetName();
 	return GUnrealEd->HasMountWritePermissionForPackage(PackageName);
 }
 
+// NOTE: This function is very similar to USourceControlHelpers::PackageFilename() but it doesn't call FindPackage() which can
+//       crash the engine when auto-saving packages.
+FString GetPackagePathname(const UPackage* InPackage)
+{
+	auto GetPathnameInternal = [](const UPackage* Package)
+	{
+		FString PackageName = Package->GetName();
+		FString Filename = Package->GetName();
+
+		// Get the filename by finding it on disk first
+		if (!FPackageName::DoesPackageExist(PackageName, &Filename))
+		{
+			// The package does not exist on disk, see if we can find it in memory and predict the file extension
+			// Only do this if the supplied package name is valid
+			const bool bIncludeReadOnlyRoots = false;
+			if (FPackageName::IsValidLongPackageName(PackageName, bIncludeReadOnlyRoots))
+			{
+				// This is a package in memory that has not yet been saved. Determine the extension and convert to a filename, if we do have the package, just assume normal asset extension
+				const FString PackageExtension = Package->ContainsMap() ? FPackageName::GetMapPackageExtension() : FPackageName::GetAssetPackageExtension();
+				Filename = FPackageName::LongPackageNameToFilename(PackageName, PackageExtension);
+			}
+		}
+
+		return Filename;
+	};
+
+	return FPaths::ConvertRelativePathToFull(GetPathnameInternal(InPackage));
+}
+
+
 } // Anonoymous namespace
 
 
@@ -108,6 +142,21 @@ FUnsavedAssetsTracker::FUnsavedAssetsTracker()
 
 	// Register for the package modified callback to catch packages that have been saved
 	UPackage::PackageSavedWithContextEvent.AddRaw(this, &FUnsavedAssetsTracker::OnPackageSaved);
+
+	// Hook to detect when a map is changed to refresh to catch when a temporary map is discarded.
+	FLevelEditorModule& LevelEditor = FModuleManager::LoadModuleChecked<FLevelEditorModule>("LevelEditor");
+	LevelEditor.OnMapChanged().AddRaw(this, &FUnsavedAssetsTracker::OnMapChanged);
+
+	// Hook to detect when a a world is renamed to to catch when a temporary map is saved with a new name.
+	FWorldDelegates::OnPostWorldRename.AddRaw(this, &FUnsavedAssetsTracker::OnWorldPostRename);
+
+	// Hook to detect when an Undo/Redo change the dirty state of a package.
+	if (GEditor != nullptr && GEditor->Trans != nullptr)
+	{
+		UTransBuffer* TransBuffer = CastChecked<UTransBuffer>(GEditor->Trans);
+		TransBuffer->OnUndo().AddRaw(this, &FUnsavedAssetsTracker::OnUndo);
+		TransBuffer->OnRedo().AddRaw(this, &FUnsavedAssetsTracker::OnRedo);
+	}
 }
 
 FUnsavedAssetsTracker::~FUnsavedAssetsTracker()
@@ -115,6 +164,20 @@ FUnsavedAssetsTracker::~FUnsavedAssetsTracker()
 	UPackage::PackageDirtyStateChangedEvent.RemoveAll(this);
 	UPackage::PackageMarkedDirtyEvent.RemoveAll(this);
 	UPackage::PackageSavedWithContextEvent.RemoveAll(this);
+
+	if (FLevelEditorModule* LevelEditor = FModuleManager::GetModulePtr<FLevelEditorModule>("LevelEditor"))
+	{
+		LevelEditor->OnMapChanged().RemoveAll(this);
+	}
+
+	FWorldDelegates::OnPostWorldRename.RemoveAll(this);
+
+	if (GEditor != nullptr && GEditor->Trans != nullptr)
+	{
+		UTransBuffer* TransBuffer = CastChecked<UTransBuffer>(GEditor->Trans);
+		TransBuffer->OnUndo().RemoveAll(this);
+		TransBuffer->OnRedo().RemoveAll(this);
+	}
 }
 
 int32 FUnsavedAssetsTracker::GetUnsavedAssetNum() const
@@ -154,34 +217,102 @@ TMap<FString, FString> FUnsavedAssetsTracker::GetWarnings() const
 
 void FUnsavedAssetsTracker::OnPackageMarkedDirty(UPackage* Package, bool bWasDirty)
 {
-	AddPackage(Package);
+	if (ShouldTrackDirtyPackage(Package))
+	{
+		StartTrackingDirtyPackage(Package);
+	}
 }
 
 void FUnsavedAssetsTracker::OnPackageDirtyStateUpdated(UPackage* Package)
 {
+	if (!ShouldTrackDirtyPackage(Package))
+	{
+		return;
+	}
+
 	if (Package->IsDirty())
 	{
-		AddPackage(Package);
+		StartTrackingDirtyPackage(Package);
 	}
 	else
 	{
-		RemovePackage(Package);
+		StopTrackingDirtyPackage(Package);
 	}
 }
 
 void FUnsavedAssetsTracker::OnPackageSaved(const FString& PackagePathname, UPackage* Package, FObjectPostSaveContext ObjectSaveContext)
 {
-	RemovePackage(Package);
+	if (ObjectSaveContext.IsProceduralSave() || (ObjectSaveContext.GetSaveFlags() & SAVE_FromAutosave) != 0)
+	{
+		return; // Don't track procedural save (during cooking) nor package auto-saved as backup in case of crash
+	}
+
+	if (ShouldTrackDirtyPackage(Package))
+	{
+		StopTrackingDirtyPackage(Package);
+	}
 }
 
-void FUnsavedAssetsTracker::AddPackage(UPackage* Package)
+void FUnsavedAssetsTracker::SyncWithDirtyPackageList()
 {
-	if (!IsPackagePersistent(Package))
+	// The the list of dirty packages tracked by the engine (considered source of truth)
+	TArray<UPackage*> DirtyPackages;
+	FEditorFileUtils::GetDirtyPackages(DirtyPackages);
+
+	TArray<FString> ToRemove;
+
+	// Remove packages that used to be dirty but aren't dirty anymore. (usually because the package was saved/renamed at the same time)
+	for (const TPair<FString, FStatus>& PathnameStatusPair : UnsavedFiles)
 	{
-		return;
+		if (!DirtyPackages.ContainsByPredicate([&PathnameStatusPair](const UPackage* Pkg) { return GetPackagePathname(Pkg) == PathnameStatusPair.Key; }))
+		{
+			ToRemove.Emplace(PathnameStatusPair.Key);
+		}
+	}
+	for (const FString& PathnameToRemove : ToRemove)
+	{
+		StopTrackingDirtyPackage(PathnameToRemove);
 	}
 
-	FString PackagePathname = SourceControlHelpers::PackageFilename(Package);
+	// Add packages that aren't tracked yet
+	for (UPackage* Package : DirtyPackages)
+	{
+		if (ShouldTrackDirtyPackage(Package))
+		{
+			StartTrackingDirtyPackage(Package); // This early out if the package is already tracked.
+		}
+	}
+}
+
+void FUnsavedAssetsTracker::OnUndo(const FTransactionContext& TransactionContext, bool Succeeded)
+{
+	SyncWithDirtyPackageList();
+}
+
+void FUnsavedAssetsTracker::OnRedo(const FTransactionContext& TransactionContext, bool Succeeded)
+{
+	SyncWithDirtyPackageList();
+}
+
+void FUnsavedAssetsTracker::OnWorldPostRename(UWorld* InWorld)
+{
+	// Saving the temporary 'Untitled' map into a package is a save/rename operation. It is simpler to
+	// sync the list of dirty package rather than implementing the rename logic, but a bit less efficient.
+	SyncWithDirtyPackageList();
+}
+
+void FUnsavedAssetsTracker::OnMapChanged(UWorld* InWorld, EMapChangeType InMapChangeType)
+{
+	// Changing map sometimes drop changes to the temporary 'Untitled' map. It is simpler to sync the
+	// list of dirty package rather than implementing the map tear down logic, but a bit less efficient.
+	SyncWithDirtyPackageList();
+}
+
+void FUnsavedAssetsTracker::StartTrackingDirtyPackage(UPackage* Package)
+{
+	checkSlow(ShouldTrackDirtyPackage(Package)); // This check should be done prior calling this function.
+
+	FString PackagePathname = GetPackagePathname(Package);
 	if (PackagePathname.IsEmpty())
 	{
 		return;
@@ -212,19 +343,21 @@ void FUnsavedAssetsTracker::AddPackage(UPackage* Package)
 	}
 }
 
-void FUnsavedAssetsTracker::RemovePackage(UPackage* Package)
+void FUnsavedAssetsTracker::StopTrackingDirtyPackage(UPackage* Package)
 {
-	if (!IsPackagePersistent(Package))
-	{
-		return;
-	}
+	checkSlow(ShouldTrackDirtyPackage(Package)); // This check should be done prior calling this function.
 
-	FString PackagePathname = SourceControlHelpers::PackageFilename(Package);
+	FString PackagePathname = GetPackagePathname(Package);
 	if (PackagePathname.IsEmpty())
 	{
 		return;
 	}
 
+	StopTrackingDirtyPackage(PackagePathname);
+}
+
+void FUnsavedAssetsTracker::StopTrackingDirtyPackage(const FString& PackagePathname)
+{
 	if (FStatus* Status = UnsavedFiles.Find(PackagePathname))
 	{
 		FString HumanFriendlyAssetName = MoveTemp(Status->HumanFriendlyAssetName); // Keep it for logging.
@@ -236,6 +369,11 @@ void FUnsavedAssetsTracker::RemovePackage(UPackage* Package)
 		WarningFiles.Remove(PackagePathname);
 		FUnsavedAssetsTrackerModule::Get().OnUnsavedAssetRemoved.Broadcast(PackagePathname);
 
+		if (WarningFiles.IsEmpty())
+		{
+			ShownWarnings.Reset();
+		}
+
 		UE_LOG(LogUnsavedAssetsTracker, Verbose, TEXT("Removed file from the unsaved asset list: %s (%s)"), *HumanFriendlyAssetName, *PackagePathname);
 	}
 }
@@ -297,6 +435,9 @@ void FUnsavedAssetsTracker::PrompToSavePackages()
 			// User likely saved something, reset the warnings. We could scan the list of unsaved asset that weren't saved (if any) and check if some warning
 			// types remain, but that looks overkill in this context.
 			ShownWarnings.Reset();
+
+			// Stay in sync with what the package the engine thinks is dirty.
+			SyncWithDirtyPackageList();
 		}
 	}
 }
diff --git a/Engine/Source/Developer/UnsavedAssetsTracker/Source/Private/UnsavedAssetsTracker.h b/Engine/Source/Developer/UnsavedAssetsTracker/Source/Private/UnsavedAssetsTracker.h
--- a/Engine/Source/Developer/UnsavedAssetsTracker/Source/Private/UnsavedAssetsTracker.h
+++ b/Engine/Source/Developer/UnsavedAssetsTracker/Source/Private/UnsavedAssetsTracker.h
@@ -5,9 +5,12 @@
 #include "CoreMinimal.h"
 #include "Templates/SharedPointer.h"
 #include "UObject/ObjectSaveContext.h"
+#include "UnrealEdMisc.h"
 
-class UPackage;
 class ISourceControlState;
+class UPackage;
+class UWorld;
+struct FTransactionContext;
 
 /**
  * Tracks assets that has in-memory modification not saved to disk yet and checks
@@ -80,11 +83,28 @@ private:
 	 */
 	void OnPackageDirtyStateUpdated(UPackage* Package);
 
+	/**
+	 * Invoked when the level editor changes map. Used to track when a temporary 'Untitled' unsaved map is closed without saving.
+	 */
+	void OnMapChanged(UWorld* InWorld, EMapChangeType InMapChangeType);
+
+	/**
+	 * Invoked when the undo stack changes. Used to track when the undo/redo affect the dirty state of a package.
+	 */
+	void OnUndo(const FTransactionContext& TransactionContext, bool Succeeded);
+	void OnRedo(const FTransactionContext& TransactionContext, bool Succeeded);
+
+	/**
+	 * Invoked when a world is successfully renamed. Used to track when a temporary 'Untitled' unsaved map is saved with a new name.
+	 */
+	void OnWorldPostRename(UWorld* World);
+
 	/** Starts to track an unsaved package.*/
-	void AddPackage(UPackage* Package);
+	void StartTrackingDirtyPackage(UPackage* Package);
 
-	/** Stops tracking the specified unsaved package. */
-	void RemovePackage(UPackage* Package);
+	/** Stops tracking the specified unsaved package if the package was previously tracked. */
+	void StopTrackingDirtyPackage(UPackage* Package);
+	void StopTrackingDirtyPackage(const FString& PackagePathname);
 
 	/** Invoked when the status of a source controlled is updated, with the corresponding warning, if any. */
 	void OnSourceControlFileStatusUpdate(const FString& Pathname, const ISourceControlState* Status);
@@ -95,6 +115,9 @@ private:
 	/** Show a toast notifications if the warning types hasn't been shown yet. */
 	void ShowWarningNotificationIfNotAlreadyShown(EWarningTypes WarningType, const FText& Msg);
 
+	/** Refresh the list of unsaved files from the list returned by FEditorFileUtils::GetDirtyPackages(). */
+	void SyncWithDirtyPackageList();
+
 private:
 	TMap<FString, FStatus> UnsavedFiles;
 	TSet<FString> WarningFiles;
