Optimized the UnsavedAssetsTracker to avoid building the full package pathname for package that are already tracked.

diff --git a/Engine/Source/Developer/UnsavedAssetsTracker/Source/Private/UnsavedAssetsTracker.cpp b/Engine/Source/Developer/UnsavedAssetsTracker/Source/Private/UnsavedAssetsTracker.cpp
--- a/Engine/Source/Developer/UnsavedAssetsTracker/Source/Private/UnsavedAssetsTracker.cpp
+++ b/Engine/Source/Developer/UnsavedAssetsTracker/Source/Private/UnsavedAssetsTracker.cpp
@@ -3,7 +3,7 @@
 #include "UnsavedAssetsTracker.h"
 
 #include "UnsavedAssetsTrackerModule.h"
-#include "AssetToolsModule.h"
+#include "Algo/Transform.h"
 #include "FileHelpers.h"
 #include "Logging/LogMacros.h"
 #include "ISourceControlModule.h"
@@ -18,6 +18,7 @@
 #include "UnrealEdGlobals.h"
 #include "Misc/PackageName.h"
 #include "LevelEditor.h"
+#include "ProfilingDebugging/CpuProfilerTrace.h"
 
 #define LOCTEXT_NAMESPACE "UnsavedAssetsTracker"
 
@@ -40,58 +41,6 @@ bool ShouldTrackDirtyPackage(const UPackage* Package)
 	return true;
 }
 
-FString GetHumanFriendlyAssetName(const UPackage* Package)
-{
-	FName AssetName;
-	FName OwnerName;
-
-	FAssetToolsModule& AssetToolsModule = FModuleManager::LoadModuleChecked<FAssetToolsModule>(TEXT("AssetTools"));
-
-	// Lookup for the first asset in the package
-	UObject* FoundAsset = nullptr;
-	ForEachObjectWithPackage(Package, [&FoundAsset](UObject* InnerObject)
-	{
-		if (InnerObject->IsAsset())
-		{
-			if (FAssetData::IsUAsset(InnerObject))
-			{
-				// If we found the primary asset, use it
-				FoundAsset = InnerObject;
-				return false;
-			}
-			// Otherwise, keep the first found asset but keep looking for a primary asset
-			if (!FoundAsset)
-			{
-				FoundAsset = InnerObject;
-			}
-		}
-		return true;
-	}, /*bIncludeNestedObjects*/ false);
-
-	if (FoundAsset)
-	{
-		TWeakPtr<IAssetTypeActions> AssetTypeActions = AssetToolsModule.Get().GetAssetTypeActionsForClass(FoundAsset->GetClass());
-		if (AssetTypeActions.IsValid())
-		{
-			AssetName = *AssetTypeActions.Pin()->GetObjectDisplayName(FoundAsset);
-		}
-		else
-		{
-			AssetName = FoundAsset->GetFName();
-		}
-
-		OwnerName = FoundAsset->GetOutermostObject()->GetFName();
-	}
-
-	// Last resort, display the package name
-	if (AssetName == NAME_None)
-	{
-		AssetName = *FPackageName::GetShortName(Package->GetFName());
-	}
-
-	return AssetName.ToString();
-}
-
 bool HasPackageWritePermissions(const UPackage* Package)
 {
 	// if we do not have write permission under the mount point for this package log an error in the message log to link to.
@@ -101,7 +50,7 @@ bool HasPackageWritePermissions(const UPackage* Package)
 
 // NOTE: This function is very similar to USourceControlHelpers::PackageFilename() but it doesn't call FindPackage() which can
 //       crash the engine when auto-saving packages.
-FString GetPackagePathname(const UPackage* InPackage)
+FString GetAbsolutePackageFilePathname(const UPackage* InPackage)
 {
 	auto GetPathnameInternal = [](const UPackage* Package)
 	{
@@ -182,29 +131,30 @@ FUnsavedAssetsTracker::~FUnsavedAssetsTracker()
 
 int32 FUnsavedAssetsTracker::GetUnsavedAssetNum() const
 {
-	return UnsavedFiles.Num();
+	check(UnsavedPackages.Num() == UnsavedAbsFilePathames.Num()); // UnsavedAbsFilePathames is use to reverse lookup UnsavedPackages, so they should have the same number of elements
+	return UnsavedPackages.Num();
 }
 
 TArray<FString> FUnsavedAssetsTracker::GetUnsavedAssets() const
 {
 	TArray<FString> Pathnames;
-	UnsavedFiles.GetKeys(Pathnames);
+	UnsavedAbsFilePathames.GetKeys(Pathnames);
 	return Pathnames;
 }
 
 int32 FUnsavedAssetsTracker::GetWarningNum() const
 {
-	return WarningFiles.Num();
+	return WarningPackagePathnames.Num();
 }
 
 TMap<FString, FString> FUnsavedAssetsTracker::GetWarnings() const
 {
 	TMap<FString, FString> Warnings;
-	for (const FString& Pathname: WarningFiles)
+	for (const FString& PackagePathname : WarningPackagePathnames)
 	{
-		if (const FStatus* Status = UnsavedFiles.Find(Pathname))
+		if (const FStatus* Status = UnsavedPackages.Find(PackagePathname))
 		{
-			Warnings.Add(Pathname, Status->WarningText.ToString());
+			Warnings.Add(Status->AbsPackageFilePathname, Status->WarningText.ToString());
 		}
 		else
 		{
@@ -236,7 +186,7 @@ void FUnsavedAssetsTracker::OnPackageDirtyStateUpdated(UPackage* Package)
 	}
 	else
 	{
-		StopTrackingDirtyPackage(Package);
+		StopTrackingDirtyPackage(Package->GetPathName());
 	}
 }
 
@@ -249,33 +199,40 @@ void FUnsavedAssetsTracker::OnPackageSaved(const FString& PackagePathname, UPack
 
 	if (ShouldTrackDirtyPackage(Package))
 	{
-		StopTrackingDirtyPackage(Package);
+		StopTrackingDirtyPackage(Package->GetPathName());
 	}
 }
 
 void FUnsavedAssetsTracker::SyncWithDirtyPackageList()
 {
+	TRACE_CPUPROFILER_EVENT_SCOPE(FUnsavedAssetsTracker::SyncWithDirtyPackageList);
+
 	// The the list of dirty packages tracked by the engine (considered source of truth)
 	TArray<UPackage*> DirtyPackages;
 	FEditorFileUtils::GetDirtyPackages(DirtyPackages);
 
+	// Convert the dirty packages pathname only once.
+	TArray<FString> DirtyPackagePathnames;
+	DirtyPackagePathnames.Reserve(DirtyPackages.Num());
+	Algo::Transform(DirtyPackages, DirtyPackagePathnames, [](const UPackage* Package) { return Package->GetPathName(); });
+
 	TArray<FString> ToRemove;
 
 	// Remove packages that used to be dirty but aren't dirty anymore. (usually because the package was saved/renamed at the same time)
-	for (const TPair<FString, FStatus>& PathnameStatusPair : UnsavedFiles)
+	for (const TPair<FString, FStatus>& PackagePathnameStatusPair : UnsavedPackages)
 	{
-		if (!DirtyPackages.ContainsByPredicate([&PathnameStatusPair](const UPackage* Pkg) { return GetPackagePathname(Pkg) == PathnameStatusPair.Key; }))
+		if (!DirtyPackagePathnames.ContainsByPredicate([&PackagePathnameStatusPair](const FString& PackagePathname) { return PackagePathname == PackagePathnameStatusPair.Key; }))
 		{
-			ToRemove.Emplace(PathnameStatusPair.Key);
+			ToRemove.Emplace(PackagePathnameStatusPair.Key);
 		}
 	}
-	for (const FString& PathnameToRemove : ToRemove)
+	for (const FString& PackagePathnameToRemove : ToRemove)
 	{
-		StopTrackingDirtyPackage(PathnameToRemove);
+		StopTrackingDirtyPackage(PackagePathnameToRemove);
 	}
 
 	// Add packages that aren't tracked yet
-	for (UPackage* Package : DirtyPackages)
+	for (const UPackage* Package : DirtyPackages)
 	{
 		if (ShouldTrackDirtyPackage(Package))
 		{
@@ -308,113 +265,112 @@ void FUnsavedAssetsTracker::OnMapChanged(UWorld* InWorld, EMapChangeType InMapCh
 	SyncWithDirtyPackageList();
 }
 
-void FUnsavedAssetsTracker::StartTrackingDirtyPackage(UPackage* Package)
+void FUnsavedAssetsTracker::StartTrackingDirtyPackage(const UPackage* Package)
 {
-	checkSlow(ShouldTrackDirtyPackage(Package)); // This check should be done prior calling this function.
+	TRACE_CPUPROFILER_EVENT_SCOPE(FUnsavedAssetsTracker::StartTrackingDirtyPackage);
 
-	FString PackagePathname = GetPackagePathname(Package);
-	if (PackagePathname.IsEmpty())
+	FString PackagePathname = Package->GetPathName();
+	if (UnsavedPackages.Contains(PackagePathname))
 	{
-		return;
+		return; // Early out if the package is already tracked.
 	}
 
-	FString HumanFriendlyAssetName = GetHumanFriendlyAssetName(Package);
-
-	int32 UnsavedNumBefore = UnsavedFiles.Num();
-	FStatus& Status = UnsavedFiles.FindOrAdd(PackagePathname, FStatus(HumanFriendlyAssetName));
-	if (UnsavedFiles.Num() > UnsavedNumBefore) // Detect if a new asset was added.
+	// The absolute package path on disk (That's a much slower operation).
+	FString AbsPackageFilePathname = GetAbsolutePackageFilePathname(Package);
+	if (AbsPackageFilePathname.IsEmpty())
 	{
-		if (!HasPackageWritePermissions(Package))
-		{
-			Status.WarningType = EWarningTypes::PackageWritePermission;
-			Status.WarningText = FText::Format(LOCTEXT("Write_Permission_Warning", "Insufficient writing permission to save {0}"), FText::FromString(Package->GetName()));
-			WarningFiles.Add(PackagePathname);
-			ShowWarningNotificationIfNotAlreadyShown(EWarningTypes::PackageWritePermission, Status.WarningText);
-		}
-
-		ISourceControlModule::Get().GetSourceControlFileStatusMonitor().StartMonitoringFile(
-			reinterpret_cast<uintptr_t>(this),
-			PackagePathname,
-			FSourceControlFileStatusMonitor::FOnSourceControlFileStatus::CreateSP(this, &FUnsavedAssetsTracker::OnSourceControlFileStatusUpdate));
-
-		FUnsavedAssetsTrackerModule::Get().OnUnsavedAssetAdded.Broadcast(PackagePathname);
-
-		UE_LOG(LogUnsavedAssetsTracker, Verbose, TEXT("Added file to the unsaved asset list: %s (%s)"), *HumanFriendlyAssetName, *PackagePathname);
+		return;
 	}
-}
 
-void FUnsavedAssetsTracker::StopTrackingDirtyPackage(UPackage* Package)
-{
-	checkSlow(ShouldTrackDirtyPackage(Package)); // This check should be done prior calling this function.
+	FStatus& Status = UnsavedPackages.Emplace(PackagePathname, FStatus(AbsPackageFilePathname));
+	UnsavedAbsFilePathames.Emplace(AbsPackageFilePathname, PackagePathname); // For fast reverse lookup.
 
-	FString PackagePathname = GetPackagePathname(Package);
-	if (PackagePathname.IsEmpty())
+	if (!HasPackageWritePermissions(Package))
 	{
-		return;
+		Status.WarningType = EWarningTypes::PackageWritePermission;
+		Status.WarningText = FText::Format(LOCTEXT("Write_Permission_Warning", "Insufficient writing permission to save {0}"), FText::FromString(Package->GetName()));
+		WarningPackagePathnames.Add(PackagePathname);
+		ShowWarningNotificationIfNotAlreadyShown(EWarningTypes::PackageWritePermission, Status.WarningText);
 	}
 
-	StopTrackingDirtyPackage(PackagePathname);
+	ISourceControlModule::Get().GetSourceControlFileStatusMonitor().StartMonitoringFile(
+		reinterpret_cast<uintptr_t>(this),
+		AbsPackageFilePathname,
+		FSourceControlFileStatusMonitor::FOnSourceControlFileStatus::CreateSP(this, &FUnsavedAssetsTracker::OnSourceControlFileStatusUpdate));
+
+	FUnsavedAssetsTrackerModule::Get().OnUnsavedAssetAdded.Broadcast(AbsPackageFilePathname);
+
+	UE_LOG(LogUnsavedAssetsTracker, Verbose, TEXT("Added file to the unsaved asset list: %s"), *PackagePathname);
 }
 
 void FUnsavedAssetsTracker::StopTrackingDirtyPackage(const FString& PackagePathname)
 {
-	if (FStatus* Status = UnsavedFiles.Find(PackagePathname))
+	TRACE_CPUPROFILER_EVENT_SCOPE(FUnsavedAssetsTracker::StopTrackingDirtyPackage);
+
+	if (FStatus* Status = UnsavedPackages.Find(PackagePathname))
 	{
-		FString HumanFriendlyAssetName = MoveTemp(Status->HumanFriendlyAssetName); // Keep it for logging.
+		FString AbsPackageFilePathname = MoveTemp(Status->AbsPackageFilePathname);
 
-		UnsavedFiles.Remove(PackagePathname);
-		ISourceControlModule::Get().GetSourceControlFileStatusMonitor().StopMonitoringFile(reinterpret_cast<uintptr_t>(this), PackagePathname);
+		UnsavedPackages.Remove(PackagePathname);
+		WarningPackagePathnames.Remove(PackagePathname);
+		UnsavedAbsFilePathames.Remove(AbsPackageFilePathname);
 
-		// Remove warnings this asset was generating (if any).
-		WarningFiles.Remove(PackagePathname);
-		FUnsavedAssetsTrackerModule::Get().OnUnsavedAssetRemoved.Broadcast(PackagePathname);
+		ISourceControlModule::Get().GetSourceControlFileStatusMonitor().StopMonitoringFile(reinterpret_cast<uintptr_t>(this), AbsPackageFilePathname);
+		FUnsavedAssetsTrackerModule::Get().OnUnsavedAssetRemoved.Broadcast(AbsPackageFilePathname);
 
-		if (WarningFiles.IsEmpty())
+		if (WarningPackagePathnames.IsEmpty())
 		{
 			ShownWarnings.Reset();
 		}
 
-		UE_LOG(LogUnsavedAssetsTracker, Verbose, TEXT("Removed file from the unsaved asset list: %s (%s)"), *HumanFriendlyAssetName, *PackagePathname);
+		UE_LOG(LogUnsavedAssetsTracker, Verbose, TEXT("Removed file from the unsaved asset list: %s"), *PackagePathname);
 	}
 }
 
-void FUnsavedAssetsTracker::OnSourceControlFileStatusUpdate(const FString& Pathname, const ISourceControlState* State)
+void FUnsavedAssetsTracker::OnSourceControlFileStatusUpdate(const FString& AbsFilePathname, const ISourceControlState* State)
 {
-	auto DiscardWarning = [this](FStatus* InStatus, const FString& Pathname)
+	auto DiscardWarning = [this](const FString& InPackagePathname, FStatus& InStatus)
 	{
 		// Source control status update cannot clear the package write permission warning.
-		if (InStatus->WarningType != EWarningTypes::PackageWritePermission)
+		if (InStatus.WarningType != EWarningTypes::PackageWritePermission)
 		{
-			WarningFiles.Remove(Pathname);
-			InStatus->WarningText = FText::GetEmpty();
-			InStatus->WarningType = EWarningTypes::None;
+			WarningPackagePathnames.Remove(InPackagePathname);
+			InStatus.WarningText = FText::GetEmpty();
+			InStatus.WarningType = EWarningTypes::None;
 		}
 
-		if (WarningFiles.IsEmpty()) // All warning were cleared.
+		if (WarningPackagePathnames.IsEmpty()) // All warning were cleared.
 		{
 			ShownWarnings.Reset(); // Reeactivate the notification next time a warning happens.
 		}
 	};
 
-	if (FStatus* Status = UnsavedFiles.Find(Pathname))
+	if (const FString* PackagePathname = UnsavedAbsFilePathames.Find(AbsFilePathname)) // Reverse lookup.
 	{
-		if (Status->WarningType == EWarningTypes::PackageWritePermission)
-		{
-			return; // Write permission issue has more weight than source control issue.
-		}
-		else if (State == nullptr) // Source control state was reset. (Changing source control provider/disabling source control)
+		if (FStatus* Status = UnsavedPackages.Find(*PackagePathname))
 		{
-			DiscardWarning(Status, Pathname);
-		}
-		else if (TOptional<FText> WarningText = State->GetWarningText())
-		{
-			Status->WarningText = *WarningText;
-			OnSourceControlWarningNotification(*State, *Status);
-			WarningFiles.Add(Pathname);
+			if (Status->WarningType == EWarningTypes::PackageWritePermission)
+			{
+				return; // Write permission issue has more weight than source control issue.
+			}
+			else if (State == nullptr) // Source control state was reset. (Changing source control provider/disabling source control)
+			{
+				DiscardWarning(*PackagePathname, *Status);
+			}
+			else if (TOptional<FText> WarningText = State->GetWarningText())
+			{
+				Status->WarningText = *WarningText;
+				OnSourceControlWarningNotification(*State, *Status);
+				WarningPackagePathnames.Emplace(*PackagePathname);
+			}
+			else
+			{
+				DiscardWarning(*PackagePathname, *Status);
+			}
 		}
 		else
 		{
-			DiscardWarning(Status, Pathname);
+			checkNoEntry(); // UnsavedAbsFilePathames values are keys in UnsavedPackages to perform reverse lookup and those must always match.
 		}
 	}
 }
diff --git a/Engine/Source/Developer/UnsavedAssetsTracker/Source/Private/UnsavedAssetsTracker.h b/Engine/Source/Developer/UnsavedAssetsTracker/Source/Private/UnsavedAssetsTracker.h
--- a/Engine/Source/Developer/UnsavedAssetsTracker/Source/Private/UnsavedAssetsTracker.h
+++ b/Engine/Source/Developer/UnsavedAssetsTracker/Source/Private/UnsavedAssetsTracker.h
@@ -55,8 +55,8 @@ private:
 
 	struct FStatus
 	{
-		FStatus(FString InHumanFriendlyAssetName) : HumanFriendlyAssetName(InHumanFriendlyAssetName) {}
-		FString HumanFriendlyAssetName;
+		FStatus(const FString& InAbsFilePathname) : AbsPackageFilePathname(InAbsFilePathname) { }
+		FString AbsPackageFilePathname; // On disk.
 		EWarningTypes WarningType = EWarningTypes::None;
 		FText WarningText;
 	};
@@ -100,10 +100,9 @@ private:
 	void OnWorldPostRename(UWorld* World);
 
 	/** Starts to track an unsaved package.*/
-	void StartTrackingDirtyPackage(UPackage* Package);
+	void StartTrackingDirtyPackage(const UPackage* Package);
 
 	/** Stops tracking the specified unsaved package if the package was previously tracked. */
-	void StopTrackingDirtyPackage(UPackage* Package);
 	void StopTrackingDirtyPackage(const FString& PackagePathname);
 
 	/** Invoked when the status of a source controlled is updated, with the corresponding warning, if any. */
@@ -119,8 +118,18 @@ private:
 	void SyncWithDirtyPackageList();
 
 private:
-	TMap<FString, FStatus> UnsavedFiles;
-	TSet<FString> WarningFiles;
+	/** Maps the package pathname to the status */
+	TMap<FString, FStatus> UnsavedPackages;
+
+	/** Maps the packages absolute file name (on file system) to the package pathname for fast reverse lookup. */
+	TMap<FString, FString> UnsavedAbsFilePathames;
+
+	/** Packages pathname for which a warning was issued. */
+	TSet<FString> WarningPackagePathnames;
+
+	/** Sets of warning shown to the user. */
 	TSet<EWarningTypes> ShownWarnings;
+
+	/** Whether toast notification should be shown when a warning is detected. */
 	bool bWarningNotificationEnabled = true;
 };
