Optimization pass on FUnsavedAssetsTracker::SyncWithDirtyPackageList()

diff --git a/Engine/Source/Developer/UnsavedAssetsTracker/Source/Private/UnsavedAssetsTracker.cpp b/Engine/Source/Developer/UnsavedAssetsTracker/Source/Private/UnsavedAssetsTracker.cpp
--- a/Engine/Source/Developer/UnsavedAssetsTracker/Source/Private/UnsavedAssetsTracker.cpp
+++ b/Engine/Source/Developer/UnsavedAssetsTracker/Source/Private/UnsavedAssetsTracker.cpp
@@ -211,24 +211,30 @@ void FUnsavedAssetsTracker::SyncWithDirtyPackageList()
 	TArray<UPackage*> DirtyPackages;
 	FEditorFileUtils::GetDirtyPackages(DirtyPackages);
 
-	// Convert the dirty packages pathname only once.
-	TArray<FString> DirtyPackagePathnames;
-	DirtyPackagePathnames.Reserve(DirtyPackages.Num());
-	Algo::Transform(DirtyPackages, DirtyPackagePathnames, [](const UPackage* Package) { return Package->GetPathName(); });
-
 	TArray<FString> ToRemove;
-
-	// Remove packages that used to be dirty but aren't dirty anymore. (usually because the package was saved/renamed at the same time)
-	for (const TPair<FString, FStatus>& PackagePathnameStatusPair : UnsavedPackages)
+	if (UnsavedFiles.Num() > 0)
 	{
-		if (!DirtyPackagePathnames.ContainsByPredicate([&PackagePathnameStatusPair](const FString& PackagePathname) { return PackagePathname == PackagePathnameStatusPair.Key; }))
+		// Accelerate lookup to some degree, ideally GetPackagePathname should not be called as that calls FileExists() which is very expensive
+		TMap<FString, UPackage*> PathnameToDirtyPackage;
+		PathnameToDirtyPackage.Reserve(DirtyPackages.Num());
+		for (UPackage* DirtyPackage : DirtyPackages)
 		{
-			ToRemove.Emplace(PackagePathnameStatusPair.Key);
+			PathnameToDirtyPackage.Add(GetPackagePathname(DirtyPackage), DirtyPackage);
+		}
+
+		// Remove packages that used to be dirty but aren't dirty anymore. (usually because the package was saved/renamed at the same time)
+		for (const TPair<FString, FStatus>& PathnameStatusPair : UnsavedFiles)
+		{
+			if (!PathnameToDirtyPackage.Contains(PathnameStatusPair.Key))
+			{
+				ToRemove.Emplace(PathnameStatusPair.Key);
+			}
+		}
+
+		for (const FString& PathnameToRemove : ToRemove)
+		{
+			StopTrackingDirtyPackage(PathnameToRemove);
 		}
-	}
-	for (const FString& PackagePathnameToRemove : ToRemove)
-	{
-		StopTrackingDirtyPackage(PackagePathnameToRemove);
 	}
 
 	// Add packages that aren't tracked yet
