Salvage a merge conflict

diff --git a/Engine/Source/Developer/UnsavedAssetsTracker/Source/Private/UnsavedAssetsTracker.cpp b/Engine/Source/Developer/UnsavedAssetsTracker/Source/Private/UnsavedAssetsTracker.cpp
--- a/Engine/Source/Developer/UnsavedAssetsTracker/Source/Private/UnsavedAssetsTracker.cpp
+++ b/Engine/Source/Developer/UnsavedAssetsTracker/Source/Private/UnsavedAssetsTracker.cpp
@@ -211,31 +211,28 @@ void FUnsavedAssetsTracker::SyncWithDirtyPackageList()
 	TArray<UPackage*> DirtyPackages;
 	FEditorFileUtils::GetDirtyPackages(DirtyPackages);
 
-	TArray<FString> ToRemove;
-	if (UnsavedFiles.Num() > 0)
+	// Convert the dirty packages pathname only once.
+	TSet<FString> DirtyPackagePathnames;
+	DirtyPackagePathnames.Reserve(DirtyPackages.Num());
+	for (UPackage* DirtyPackage : DirtyPackages)
 	{
-		// Accelerate lookup to some degree, ideally GetPackagePathname should not be called as that calls FileExists() which is very expensive
-		TMap<FString, UPackage*> PathnameToDirtyPackage;
-		PathnameToDirtyPackage.Reserve(DirtyPackages.Num());
-		for (UPackage* DirtyPackage : DirtyPackages)
-		{
-			PathnameToDirtyPackage.Add(GetPackagePathname(DirtyPackage), DirtyPackage);
-		}
+		DirtyPackagePathnames.Add(Package->GetPathName());
+	}
 
-		// Remove packages that used to be dirty but aren't dirty anymore. (usually because the package was saved/renamed at the same time)
-		for (const TPair<FString, FStatus>& PathnameStatusPair : UnsavedFiles)
-		{
-			if (!PathnameToDirtyPackage.Contains(PathnameStatusPair.Key))
-			{
-				ToRemove.Emplace(PathnameStatusPair.Key);
-			}
-		}
+	TArray<FString> ToRemove;
 
-		for (const FString& PathnameToRemove : ToRemove)
+	// Remove packages that used to be dirty but aren't dirty anymore. (usually because the package was saved/renamed at the same time)
+	for (const TPair<FString, FStatus>& PackagePathnameStatusPair : UnsavedPackages)
+	{
+		if (!DirtyPackagePathnames.Contains(PackagePathnameStatusPair.Key))
 		{
-			StopTrackingDirtyPackage(PathnameToRemove);
+			ToRemove.Emplace(PackagePathnameStatusPair.Key);
 		}
 	}
+	for (const FString& PackagePathnameToRemove : ToRemove)
+	{
+		StopTrackingDirtyPackage(PackagePathnameToRemove);
+	}
 
 	// Add packages that aren't tracked yet
 	for (const UPackage* Package : DirtyPackages)
