CIS fix

diff --git a/Engine/Source/Developer/UnsavedAssetsTracker/Source/Private/UnsavedAssetsTracker.cpp b/Engine/Source/Developer/UnsavedAssetsTracker/Source/Private/UnsavedAssetsTracker.cpp
--- a/Engine/Source/Developer/UnsavedAssetsTracker/Source/Private/UnsavedAssetsTracker.cpp
+++ b/Engine/Source/Developer/UnsavedAssetsTracker/Source/Private/UnsavedAssetsTracker.cpp
@@ -216,7 +216,7 @@ void FUnsavedAssetsTracker::SyncWithDirtyPackageList()
 	DirtyPackagePathnames.Reserve(DirtyPackages.Num());
 	for (UPackage* DirtyPackage : DirtyPackages)
 	{
-		DirtyPackagePathnames.Add(Package->GetPathName());
+		DirtyPackagePathnames.Add(DirtyPackage->GetPathName());
 	}
 
 	TArray<FString> ToRemove;
