Fixed the UnsavedAssetsTracker being invoked on a background thread and firering delegates to code that assumes everything happens on the game thread.

diff --git a/Engine/Source/Developer/UnsavedAssetsTracker/Source/Private/UnsavedAssetsTracker.cpp b/Engine/Source/Developer/UnsavedAssetsTracker/Source/Private/UnsavedAssetsTracker.cpp
--- a/Engine/Source/Developer/UnsavedAssetsTracker/Source/Private/UnsavedAssetsTracker.cpp
+++ b/Engine/Source/Developer/UnsavedAssetsTracker/Source/Private/UnsavedAssetsTracker.cpp
@@ -17,6 +17,7 @@
 #include "Engine/World.h"
 #include "UnrealEdGlobals.h"
 #include "Misc/PackageName.h"
+#include "Misc/ScopeLock.h"
 #include "LevelEditor.h"
 #include "ProfilingDebugging/CpuProfilerTrace.h"
 #include "IAssetTools.h"
@@ -97,6 +98,9 @@ FUnsavedAssetsTracker::FUnsavedAssetsTracker()
 	FLevelEditorModule& LevelEditor = FModuleManager::LoadModuleChecked<FLevelEditorModule>("LevelEditor");
 	LevelEditor.OnMapChanged().AddRaw(this, &FUnsavedAssetsTracker::OnMapChanged);
 
+	// Register in the ticker.
+	TickerHandle = FTSTicker::GetCoreTicker().AddTicker(FTickerDelegate::CreateRaw(this, &FUnsavedAssetsTracker::Tick));
+
 	// Hook to detect when a a world is renamed to to catch when a temporary map is saved with a new name.
 	FWorldDelegates::OnPostWorldRename.AddRaw(this, &FUnsavedAssetsTracker::OnWorldPostRename);
 
@@ -111,6 +115,8 @@ FUnsavedAssetsTracker::FUnsavedAssetsTracker()
 
 FUnsavedAssetsTracker::~FUnsavedAssetsTracker()
 {
+	FTSTicker::GetCoreTicker().RemoveTicker(TickerHandle);
+
 	UPackage::PackageDirtyStateChangedEvent.RemoveAll(this);
 	UPackage::PackageMarkedDirtyEvent.RemoveAll(this);
 	UPackage::PackageSavedWithContextEvent.RemoveAll(this);
@@ -170,7 +176,15 @@ void FUnsavedAssetsTracker::OnPackageMarkedDirty(UPackage* Package, bool bWasDir
 {
 	if (ShouldTrackDirtyPackage(Package))
 	{
-		StartTrackingDirtyPackage(Package);
+		if (IsInGameThread())
+		{
+			StartTrackingDirtyPackage(Package);
+		}
+		else
+		{
+			FScopeLock Lock(&DirtyNotificationLock);
+			DirtyPackageEventsReportedOnBackgroundThread.Emplace(Package->GetPathName(), Package);
+		}
 	}
 }
 
@@ -181,13 +195,21 @@ void FUnsavedAssetsTracker::OnPackageDirtyStateUpdated(UPackage* Package)
 		return;
 	}
 
-	if (Package->IsDirty())
+	if (IsInGameThread())
 	{
-		StartTrackingDirtyPackage(Package);
+		if (Package->IsDirty())
+		{
+			StartTrackingDirtyPackage(Package);
+		}
+		else
+		{
+			StopTrackingDirtyPackage(Package->GetPathName());
+		}
 	}
 	else
 	{
-		StopTrackingDirtyPackage(Package->GetPathName());
+		FScopeLock Lock(&DirtyNotificationLock);
+		DirtyPackageEventsReportedOnBackgroundThread.Emplace(Package->GetPathName(), Package);
 	}
 }
 
@@ -454,4 +476,31 @@ void FUnsavedAssetsTracker::ShowWarningNotificationIfNotAlreadyShown(EWarningTyp
 	}
 }
 
+bool FUnsavedAssetsTracker::Tick(float DeltaTime)
+{
+	if (DirtyNotificationLock.TryLock()) // Don't block the main thread.
+	{
+		for (TPair<FString, TWeakObjectPtr<UPackage>>& PathNamePackagePair : DirtyPackageEventsReportedOnBackgroundThread)
+		{
+			if (PathNamePackagePair.Value.IsValid())
+			{
+				UPackage* Package = PathNamePackagePair.Value.Get();
+				if (Package->IsDirty())
+				{
+					StartTrackingDirtyPackage(Package);
+				}
+				else
+				{
+					StopTrackingDirtyPackage(Package->GetPathName());
+				}
+			}
+		}
+
+		DirtyPackageEventsReportedOnBackgroundThread.Reset();
+		DirtyNotificationLock.Unlock();
+	}
+
+	return true; // Tick again.
+}
+
 #undef LOCTEXT_NAMESPACE
diff --git a/Engine/Source/Developer/UnsavedAssetsTracker/Source/Private/UnsavedAssetsTracker.h b/Engine/Source/Developer/UnsavedAssetsTracker/Source/Private/UnsavedAssetsTracker.h
--- a/Engine/Source/Developer/UnsavedAssetsTracker/Source/Private/UnsavedAssetsTracker.h
+++ b/Engine/Source/Developer/UnsavedAssetsTracker/Source/Private/UnsavedAssetsTracker.h
@@ -3,6 +3,8 @@
 #pragma once
 
 #include "CoreMinimal.h"
+#include "Containers/Ticker.h"
+#include "HAL/CriticalSection.h"
 #include "Templates/SharedPointer.h"
 #include "UObject/ObjectSaveContext.h"
 #include "UnrealEdMisc.h"
@@ -117,7 +119,16 @@ private:
 	/** Refresh the list of unsaved files from the list returned by FEditorFileUtils::GetDirtyPackages(). */
 	void SyncWithDirtyPackageList();
 
+	/** Invoked by the engine on the game thread. */
+	bool Tick(float DeltaTime);
+
 private:
+	/** Supports multithreaded Dirty notifications. */
+	FCriticalSection DirtyNotificationLock;
+
+	/** Keep the packages that were marked/unmarked dirty from a background thread to check the status on the game thread. */
+	TMap<FString, TWeakObjectPtr<UPackage>> DirtyPackageEventsReportedOnBackgroundThread;
+
 	/** Maps the package pathname to the status */
 	TMap<FString, FStatus> UnsavedPackages;
 
@@ -130,6 +141,9 @@ private:
 	/** Sets of warning shown to the user. */
 	TSet<EWarningTypes> ShownWarnings;
 
+	/** Handle to unregister the tick delegate. */
+	FTSTicker::FDelegateHandle TickerHandle;
+
 	/** Whether toast notification should be shown when a warning is detected. */
 	bool bWarningNotificationEnabled = true;
 };
