Fixed UnsavedAssetTracker to display the 'Unsaved' count that matches what the engine considers 'to be saved'.
  - Fixed cases when working on an /Temp map package, typically an 'Untitled' map.
  - Fixed cases when adding/deleting an actor from the map.

diff --git a/Engine/Source/Developer/UnsavedAssetsTracker/Source/Private/UnsavedAssetsTracker.cpp b/Engine/Source/Developer/UnsavedAssetsTracker/Source/Private/UnsavedAssetsTracker.cpp
--- a/Engine/Source/Developer/UnsavedAssetsTracker/Source/Private/UnsavedAssetsTracker.cpp
+++ b/Engine/Source/Developer/UnsavedAssetsTracker/Source/Private/UnsavedAssetsTracker.cpp
@@ -4,6 +4,7 @@
 
 #include "UnsavedAssetsTrackerModule.h"
 #include "Algo/Transform.h"
+#include "AssetToolsModule.h"
 #include "FileHelpers.h"
 #include "Logging/LogMacros.h"
 #include "ISourceControlModule.h"
@@ -14,6 +15,7 @@
 #include "Editor.h"
 #include "Editor/UnrealEdEngine.h"
 #include "Editor/TransBuffer.h"
+#include "Engine/Engine.h"
 #include "Engine/World.h"
 #include "UnrealEdGlobals.h"
 #include "Misc/PackageName.h"
@@ -29,7 +31,59 @@ DEFINE_LOG_CATEGORY_STATIC(LogUnsavedAssetsTracker, Log, All);
 namespace
 {
 
-bool ShouldTrackDirtyPackage(const UPackage* Package)
+FString GetHumanFriendlyAssetName(const UPackage* Package)
+{
+	FName AssetName;
+	FName OwnerName;
+
+	FAssetToolsModule& AssetToolsModule = FModuleManager::LoadModuleChecked<FAssetToolsModule>(TEXT("AssetTools"));
+
+	// Lookup for the first asset in the package
+	UObject* FoundAsset = nullptr;
+	ForEachObjectWithPackage(Package, [&FoundAsset](UObject* InnerObject)
+	{
+		if (InnerObject->IsAsset())
+		{
+			if (FAssetData::IsUAsset(InnerObject))
+			{
+				// If we found the primary asset, use it
+				FoundAsset = InnerObject;
+				return false;
+			}
+			// Otherwise, keep the first found asset but keep looking for a primary asset
+			if (!FoundAsset)
+			{
+				FoundAsset = InnerObject;
+			}
+		}
+		return true;
+	}, /*bIncludeNestedObjects*/ false);
+
+	if (FoundAsset)
+	{
+		TWeakPtr<IAssetTypeActions> AssetTypeActions = AssetToolsModule.Get().GetAssetTypeActionsForClass(FoundAsset->GetClass());
+		if (AssetTypeActions.IsValid())
+		{
+			AssetName = *AssetTypeActions.Pin()->GetObjectDisplayName(FoundAsset);
+		}
+		else
+		{
+			AssetName = FoundAsset->GetFName();
+		}
+
+		OwnerName = FoundAsset->GetOutermostObject()->GetFName();
+	}
+
+	// Last resort, display the package name
+	if (AssetName == NAME_None)
+	{
+		AssetName = *FPackageName::GetShortName(Package->GetFName());
+	}
+
+	return AssetName.ToString();
+}
+
+bool IsPackageMeantToBeSaved(const UPackage* Package)
 {
 	// Ignore the packages that aren't meant to be persistent.
 	if (Package->HasAnyFlags(RF_Transient) ||
@@ -94,6 +148,13 @@ FUnsavedAssetsTracker::FUnsavedAssetsTracker()
 	// Register for the package modified callback to catch packages that have been saved
 	UPackage::PackageSavedWithContextEvent.AddRaw(this, &FUnsavedAssetsTracker::OnPackageSaved);
 
+	// Register to get notified when something gets deleted
+	FEditorDelegates::OnPackageDeleted.AddRaw(this, &FUnsavedAssetsTracker::OnPackageDeleted);
+	if (GEngine)
+	{
+		GEngine->OnLevelActorDeleted().AddRaw(this, &FUnsavedAssetsTracker::OnActorDeleted);
+	}
+
 	// Hook to detect when a map is changed to refresh to catch when a temporary map is discarded.
 	FLevelEditorModule& LevelEditor = FModuleManager::LoadModuleChecked<FLevelEditorModule>("LevelEditor");
 	LevelEditor.OnMapChanged().AddRaw(this, &FUnsavedAssetsTracker::OnMapChanged);
@@ -120,6 +181,12 @@ FUnsavedAssetsTracker::~FUnsavedAssetsTracker()
 	UPackage::PackageDirtyStateChangedEvent.RemoveAll(this);
 	UPackage::PackageMarkedDirtyEvent.RemoveAll(this);
 	UPackage::PackageSavedWithContextEvent.RemoveAll(this);
+	FEditorDelegates::OnPackageDeleted.RemoveAll(this);
+
+	if (GEngine)
+	{
+		GEngine->OnLevelActorDeleted().RemoveAll(this);
+	}
 
 	if (FLevelEditorModule* LevelEditor = FModuleManager::GetModulePtr<FLevelEditorModule>("LevelEditor"))
 	{
@@ -172,6 +239,26 @@ TMap<FString, FString> FUnsavedAssetsTracker::GetWarnings() const
 	return Warnings;
 }
 
+bool FUnsavedAssetsTracker::ShouldTrackDirtyPackage(const UPackage* Package)
+{
+	// Ignore the packages that aren't meant to be persistent.
+	if (!IsPackageMeantToBeSaved(Package))
+	{
+		return false;
+	}
+	else if (FPackageName::IsTempPackage(Package->GetPathName()))
+	{
+		// When working on an 'Untitled' map, the map is in a temp packages until saved. For a OFPA map,
+		// added actors to the temp map are counted dirty, but when saving, the engine only propose to save
+		// the dirty worlds which create a desync in numbers displayed by the widget and the save dialog. For
+		// simplicity, when we have dirty temp package, just look at the source of truth and sync with it.
+		bSyncWithDirtyPackageList = true;
+		return false;
+	}
+	
+	return true;
+}
+
 void FUnsavedAssetsTracker::OnPackageMarkedDirty(UPackage* Package, bool bWasDirty)
 {
 	if (ShouldTrackDirtyPackage(Package))
@@ -226,6 +313,19 @@ void FUnsavedAssetsTracker::OnPackageSaved(const FString& PackagePathname, UPack
 	}
 }
 
+void FUnsavedAssetsTracker::OnPackageDeleted(UPackage* Package)
+{
+	if (UnsavedPackages.Contains(Package->GetPathName()))
+	{
+		bSyncWithDirtyPackageList = true;
+	}
+}
+
+void FUnsavedAssetsTracker::OnActorDeleted(AActor* Actor)
+{
+	bSyncWithDirtyPackageList = true;
+}
+
 void FUnsavedAssetsTracker::SyncWithDirtyPackageList()
 {
 	TRACE_CPUPROFILER_EVENT_SCOPE(FUnsavedAssetsTracker::SyncWithDirtyPackageList);
@@ -260,7 +360,7 @@ void FUnsavedAssetsTracker::SyncWithDirtyPackageList()
 	// Add packages that aren't tracked yet
 	for (const UPackage* Package : DirtyPackages)
 	{
-		if (ShouldTrackDirtyPackage(Package))
+		if (IsPackageMeantToBeSaved(Package))
 		{
 			StartTrackingDirtyPackage(Package); // This early out if the package is already tracked.
 		}
@@ -269,26 +369,26 @@ void FUnsavedAssetsTracker::SyncWithDirtyPackageList()
 
 void FUnsavedAssetsTracker::OnUndo(const FTransactionContext& TransactionContext, bool Succeeded)
 {
-	SyncWithDirtyPackageList();
+	bSyncWithDirtyPackageList = true;
 }
 
 void FUnsavedAssetsTracker::OnRedo(const FTransactionContext& TransactionContext, bool Succeeded)
 {
-	SyncWithDirtyPackageList();
+	bSyncWithDirtyPackageList = true;
 }
 
 void FUnsavedAssetsTracker::OnWorldPostRename(UWorld* InWorld)
 {
 	// Saving the temporary 'Untitled' map into a package is a save/rename operation. It is simpler to
 	// sync the list of dirty package rather than implementing the rename logic, but a bit less efficient.
-	SyncWithDirtyPackageList();
+	bSyncWithDirtyPackageList = true;
 }
 
 void FUnsavedAssetsTracker::OnMapChanged(UWorld* InWorld, EMapChangeType InMapChangeType)
 {
 	// Changing map sometimes drop changes to the temporary 'Untitled' map. It is simpler to sync the
 	// list of dirty package rather than implementing the map tear down logic, but a bit less efficient.
-	SyncWithDirtyPackageList();
+	bSyncWithDirtyPackageList = true;
 }
 
 void FUnsavedAssetsTracker::StartTrackingDirtyPackage(const UPackage* Package)
@@ -326,7 +426,7 @@ void FUnsavedAssetsTracker::StartTrackingDirtyPackage(const UPackage* Package)
 
 	FUnsavedAssetsTrackerModule::Get().OnUnsavedAssetAdded.Broadcast(AbsPackageFilePathname);
 
-	UE_LOG(LogUnsavedAssetsTracker, Verbose, TEXT("Added file to the unsaved asset list: %s"), *PackagePathname);
+	UE_LOG(LogUnsavedAssetsTracker, Verbose, TEXT("Added file to the unsaved asset list: %s (%s)"), *GetHumanFriendlyAssetName(Package), *PackagePathname);
 }
 
 void FUnsavedAssetsTracker::StopTrackingDirtyPackage(const FString& PackagePathname)
@@ -505,6 +605,12 @@ bool FUnsavedAssetsTracker::Tick(float DeltaTime)
 		DirtyNotificationLock.Unlock();
 	}
 
+	if (bSyncWithDirtyPackageList)
+	{
+		SyncWithDirtyPackageList();
+		bSyncWithDirtyPackageList = false;
+	}
+
 	return true; // Tick again.
 }
 
diff --git a/Engine/Source/Developer/UnsavedAssetsTracker/Source/Private/UnsavedAssetsTracker.h b/Engine/Source/Developer/UnsavedAssetsTracker/Source/Private/UnsavedAssetsTracker.h
--- a/Engine/Source/Developer/UnsavedAssetsTracker/Source/Private/UnsavedAssetsTracker.h
+++ b/Engine/Source/Developer/UnsavedAssetsTracker/Source/Private/UnsavedAssetsTracker.h
@@ -9,6 +9,7 @@
 #include "UObject/ObjectSaveContext.h"
 #include "UnrealEdMisc.h"
 
+class AActor;
 class ISourceControlState;
 class UPackage;
 class UWorld;
@@ -99,9 +100,13 @@ private:
 	void OnUndo(const FTransactionContext& TransactionContext, bool Succeeded);
 	void OnRedo(const FTransactionContext& TransactionContext, bool Succeeded);
 
-	/**
-	 * Invoked when a world is successfully renamed. Used to track when a temporary 'Untitled' unsaved map is saved with a new name.
-	 */
+	/** Invoked when a package is deleted. */
+	void OnPackageDeleted(UPackage* Package);
+
+	/** Invoked when an actor is deleted. */
+	void OnActorDeleted(AActor* Actor);
+
+	/** Invoked when a world is successfully renamed. Used to track when a temporary 'Untitled' unsaved map is saved with a new name. */
 	void OnWorldPostRename(UWorld* World);
 
 	/** Starts to track an unsaved package.*/
@@ -125,6 +130,9 @@ private:
 	/** Invoked by the engine on the game thread. */
 	bool Tick(float DeltaTime);
 
+	/** Whether the specified package should be tracked. */
+	bool ShouldTrackDirtyPackage(const UPackage* Package);
+
 private:
 	/** Supports multithreaded Dirty notifications. */
 	FCriticalSection DirtyNotificationLock;
@@ -149,4 +157,7 @@ private:
 
 	/** Whether toast notification should be shown when a warning is detected. */
 	bool bWarningNotificationEnabled = true;
+
+	/** Whether the tracker should sync its dirty packages list with what the engine considers dirty. */
+	bool bSyncWithDirtyPackageList = false;
 };
