Async background checkout now a component of unsaved assets tracker

Provides additional delegate hooks for automatic checkout

diff --git a/Engine/Source/Developer/UnsavedAssetsTracker/Source/Private/UnsavedAssetsAutoCheckout.cpp b/Engine/Source/Developer/UnsavedAssetsTracker/Source/Private/UnsavedAssetsAutoCheckout.cpp
new file mode 100644
--- /dev/null
+++ b/Engine/Source/Developer/UnsavedAssetsTracker/Source/Private/UnsavedAssetsAutoCheckout.cpp
@@ -0,0 +1,49 @@
+﻿// Copyright Epic Games, Inc. All Rights Reserved.
+
+#include "UnsavedAssetsAutoCheckout.h"
+
+#include "ISourceControlModule.h"
+#include "Settings/EditorLoadingSavingSettings.h"
+#include "SourceControlOperations.h"
+#include "UnsavedAssetsTrackerModule.h"
+
+FUnsavedAssetsAutoCheckout::FUnsavedAssetsAutoCheckout(FUnsavedAssetsTrackerModule* Module)
+{
+	Module->OnUnsavedAssetAdded.AddRaw(this, &FUnsavedAssetsAutoCheckout::AsyncCheckout);
+	AsyncCheckoutComplete.BindRaw(this, &FUnsavedAssetsAutoCheckout::OnAsyncCheckoutComplete);
+}
+
+FUnsavedAssetsAutoCheckout::~FUnsavedAssetsAutoCheckout()
+{
+	FUnsavedAssetsTrackerModule::Get().OnUnsavedAssetAdded.RemoveAll(this);
+}
+
+void FUnsavedAssetsAutoCheckout::AsyncCheckout(const FString& AbsoluteAssetFilepath)
+{
+	const UEditorLoadingSavingSettings* Settings = GetDefault<UEditorLoadingSavingSettings>();
+	if (!Settings->bAutomaticallyCheckoutOnAssetModification)
+	{
+		return;
+	}
+	ISourceControlProvider& SourceControlProvider = ISourceControlModule::Get().GetProvider();
+	FSourceControlOperationRef CheckOutOperation = ISourceControlOperation::Create<FCheckOut>();
+
+	OperationToPath.Add(&CheckOutOperation.Get()) = AbsoluteAssetFilepath;
+	
+	FUnsavedAssetsTrackerModule::Get().PreUnsavedAssetAutoCheckout.Broadcast(AbsoluteAssetFilepath, CheckOutOperation);
+	
+	SourceControlProvider.Execute(CheckOutOperation, AbsoluteAssetFilepath, EConcurrency::Asynchronous, AsyncCheckoutComplete);
+}
+
+void FUnsavedAssetsAutoCheckout::OnAsyncCheckoutComplete(const FSourceControlOperationRef& CheckOutOperation, ECommandResult::Type Result)
+{
+	FString const& AbsoluteAssetFilepath = OperationToPath.FindAndRemoveChecked(&CheckOutOperation.Get());
+	if (Result == ECommandResult::Succeeded)
+	{
+		FUnsavedAssetsTrackerModule::Get().PostUnsavedAssetAutoCheckout.Broadcast(AbsoluteAssetFilepath, CheckOutOperation);	
+	}
+	else if (Result == ECommandResult::Failed)
+	{
+		FUnsavedAssetsTrackerModule::Get().PostUnsavedAssetAutoCheckoutFailure.Broadcast(AbsoluteAssetFilepath, CheckOutOperation);
+	}
+}
diff --git a/Engine/Source/Developer/UnsavedAssetsTracker/Source/Private/UnsavedAssetsAutoCheckout.h b/Engine/Source/Developer/UnsavedAssetsTracker/Source/Private/UnsavedAssetsAutoCheckout.h
new file mode 100644
--- /dev/null
+++ b/Engine/Source/Developer/UnsavedAssetsTracker/Source/Private/UnsavedAssetsAutoCheckout.h
@@ -0,0 +1,27 @@
+﻿// Copyright Epic Games, Inc. All Rights Reserved.
+
+#pragma once
+
+#include "Templates/SharedPointer.h"
+#include "Containers/Map.h"
+#include "ISourceControlProvider.h"
+
+class ISourceControlOperation;
+class FUnsavedAssetsTrackerModule;
+
+/**
+ * Uses the source control provider to attempt automatic checkouts of unsaved assets.
+ */
+class FUnsavedAssetsAutoCheckout : public TSharedFromThis<FUnsavedAssetsAutoCheckout>
+{
+public:
+	explicit FUnsavedAssetsAutoCheckout(FUnsavedAssetsTrackerModule* Module);
+	~FUnsavedAssetsAutoCheckout();
+private:
+	void AsyncCheckout(const FString& AbsoluteAssetFilepath);
+	
+	FSourceControlOperationComplete AsyncCheckoutComplete;
+	void OnAsyncCheckoutComplete(const FSourceControlOperationRef&, ECommandResult::Type);
+
+	TMap<ISourceControlOperation*, FString> OperationToPath;
+};
diff --git a/Engine/Source/Developer/UnsavedAssetsTracker/Source/Private/UnsavedAssetsTrackerModule.cpp b/Engine/Source/Developer/UnsavedAssetsTracker/Source/Private/UnsavedAssetsTrackerModule.cpp
--- a/Engine/Source/Developer/UnsavedAssetsTracker/Source/Private/UnsavedAssetsTrackerModule.cpp
+++ b/Engine/Source/Developer/UnsavedAssetsTracker/Source/Private/UnsavedAssetsTrackerModule.cpp
@@ -3,15 +3,20 @@
 #include "UnsavedAssetsTrackerModule.h"
 
 #include "CoreGlobals.h"
+#include "Settings/EditorLoadingSavingSettings.h"
+#include "ISourceControlModule.h"
+#include "ISourceControlProvider.h"
+#include "SourceControlOperations.h"
 #include "Misc/App.h"
 #include "UnsavedAssetsTracker.h"
 #include "SUnsavedAssetsStatusBarWidget.h"
+#include "UnsavedAssetsAutoCheckout.h"
 #include "Widgets/SNullWidget.h"
 
-
 void FUnsavedAssetsTrackerModule::StartupModule()
 {
 	UnsavedAssetTracker = MakeShared<FUnsavedAssetsTracker>();
+	UnsavedAssetAutoCheckout = MakeShared<FUnsavedAssetsAutoCheckout>(this);
 }
 
 void FUnsavedAssetsTrackerModule::ShutdownModule()
diff --git a/Engine/Source/Developer/UnsavedAssetsTracker/Source/Public/UnsavedAssetsTrackerModule.h b/Engine/Source/Developer/UnsavedAssetsTracker/Source/Public/UnsavedAssetsTrackerModule.h
--- a/Engine/Source/Developer/UnsavedAssetsTracker/Source/Public/UnsavedAssetsTrackerModule.h
+++ b/Engine/Source/Developer/UnsavedAssetsTracker/Source/Public/UnsavedAssetsTrackerModule.h
@@ -4,6 +4,7 @@
 
 #include "CoreMinimal.h"
 #include "EditorSubsystem.h"
+#include "ISourceControlOperation.h"
 #include "Modules/ModuleInterface.h"
 #include "Modules/ModuleManager.h"
 #include "Delegates/DelegateCombinations.h"
@@ -11,6 +12,7 @@
 #include "UObject/ObjectSaveContext.h"
 
 class FUnsavedAssetsTracker;
+class FUnsavedAssetsAutoCheckout;
 class SWidget;
 
 /** Invoked when an asset was modified in memory and added to the tracker list. */
@@ -19,6 +21,11 @@ DECLARE_MULTICAST_DELEGATE_OneParam(FOnUnsavedAssetAdded, const FString& /*FileA
 /** Invoked when an asset modified in memory was saved and removed from the tracker list. */
 DECLARE_MULTICAST_DELEGATE_OneParam(FOnUnsavedAssetRemoved, const FString& /*FileAbsPathname*/);
 
+/** Invoked before an asset is potentially automatically checked out of source control */
+DECLARE_MULTICAST_DELEGATE_TwoParams(FPreUnsavedAssetAutoCheckout, const FString& /*FileAbsPathname*/, FSourceControlOperationRef& /*CheckoutOperation*/);
+
+/** Invoked after an asset has been checked out of source control */
+DECLARE_MULTICAST_DELEGATE_TwoParams(FPostUnsavedAssetAutoCheckout, const FString& /*FileAbsPathname*/, FSourceControlOperationRef const& /*CheckoutOperation*/);
 
 /**
  * Tracks assets that has in-memory modification not saved to disk yet and checks
@@ -55,8 +62,18 @@ public:
 	FOnUnsavedAssetAdded OnUnsavedAssetAdded;
 
 	/** Invoked when a file is removed from the unsaved list. */
-	FOnUnsavedAssetAdded OnUnsavedAssetRemoved;
+	FOnUnsavedAssetRemoved OnUnsavedAssetRemoved;
+
+	/** Invoked before an asset is potentially automatically checked out of source control */
+	FPreUnsavedAssetAutoCheckout PreUnsavedAssetAutoCheckout;
+	
+	/** Invoked after an asset has been checked out of source control */
+	FPostUnsavedAssetAutoCheckout PostUnsavedAssetAutoCheckout;
 
+	/** Invoked after an asset failed to be checked out of source control */
+	FPostUnsavedAssetAutoCheckout PostUnsavedAssetAutoCheckoutFailure;
+	
 private:
 	TSharedPtr<FUnsavedAssetsTracker> UnsavedAssetTracker;
+	TSharedPtr<FUnsavedAssetsAutoCheckout> UnsavedAssetAutoCheckout;
 };
diff --git a/Engine/Source/Editor/UnrealEd/Private/UnrealEdEngine.cpp b/Engine/Source/Editor/UnrealEd/Private/UnrealEdEngine.cpp
--- a/Engine/Source/Editor/UnrealEd/Private/UnrealEdEngine.cpp
+++ b/Engine/Source/Editor/UnrealEd/Private/UnrealEdEngine.cpp
@@ -578,7 +578,7 @@ void UUnrealEdEngine::OnPackageDirtyStateUpdated( UPackage* Pkg)
 
 			const UEditorLoadingSavingSettings* Settings = GetDefault<UEditorLoadingSavingSettings>();
 			if (!bAlreadyAsked && // Don't ask if we already asked once!
-				(Settings->bPromptForCheckoutOnAssetModification || Settings->bAutomaticallyCheckoutOnAssetModification))
+				Settings->bPromptForCheckoutOnAssetModification)
 			{
 				PackageToNotifyState.Add(Package, NS_Updating);
 			}
