Fix auto checkout crash on shutdown

diff --git a/Engine/Source/Developer/UnsavedAssetsTracker/Source/Private/UnsavedAssetsAutoCheckout.cpp b/Engine/Source/Developer/UnsavedAssetsTracker/Source/Private/UnsavedAssetsAutoCheckout.cpp
--- a/Engine/Source/Developer/UnsavedAssetsTracker/Source/Private/UnsavedAssetsAutoCheckout.cpp
+++ b/Engine/Source/Developer/UnsavedAssetsTracker/Source/Private/UnsavedAssetsAutoCheckout.cpp
@@ -15,7 +15,6 @@ FUnsavedAssetsAutoCheckout::FUnsavedAssetsAutoCheckout(FUnsavedAssetsTrackerModu
 
 FUnsavedAssetsAutoCheckout::~FUnsavedAssetsAutoCheckout()
 {
-	FUnsavedAssetsTrackerModule::Get().OnUnsavedAssetAdded.RemoveAll(this);
 }
 
 void FUnsavedAssetsAutoCheckout::AsyncCheckout(const FString& AbsoluteAssetFilepath)
