Imported changes submitted by [REDACTED_PERSON] in CL 23186123 in UE5/Main

diff --git a/Engine/Source/Developer/UnsavedAssetsTracker/Source/Private/UnsavedAssetsAutoCheckout.cpp b/Engine/Source/Developer/UnsavedAssetsTracker/Source/Private/UnsavedAssetsAutoCheckout.cpp
--- a/Engine/Source/Developer/UnsavedAssetsTracker/Source/Private/UnsavedAssetsAutoCheckout.cpp
+++ b/Engine/Source/Developer/UnsavedAssetsTracker/Source/Private/UnsavedAssetsAutoCheckout.cpp
@@ -15,6 +15,7 @@ FUnsavedAssetsAutoCheckout::FUnsavedAssetsAutoCheckout(FUnsavedAssetsTrackerModu
 
 FUnsavedAssetsAutoCheckout::~FUnsavedAssetsAutoCheckout()
 {
+	FUnsavedAssetsTrackerModule::Get().OnUnsavedAssetAdded.RemoveAll(this);
 }
 
 void FUnsavedAssetsAutoCheckout::AsyncCheckout(const FString& AbsoluteAssetFilepath)
