SourceControl: Fixed editor locking when dragging a large prefab into the scene while auto-checkout is enabled.

If the prefab consisted of 6000 files, then the UnsavedAssetsAutoCheckout would schedule 6000 FCheckOut operations, so it batches everything together for a single Tick.
It also scheduled the operation, regardless of whether the file is already saved to disk or not (which in case of dragging a prefab into the scene is not the case), so it checks for file existance now too.

Ideally, we would not batch per tick, but per 'user action'. However I'm not sure if we can actually use TransactionFinalized for that?

diff --git a/Engine/Source/Developer/UnsavedAssetsTracker/Source/Private/UnsavedAssetsAutoCheckout.cpp b/Engine/Source/Developer/UnsavedAssetsTracker/Source/Private/UnsavedAssetsAutoCheckout.cpp
--- a/Engine/Source/Developer/UnsavedAssetsTracker/Source/Private/UnsavedAssetsAutoCheckout.cpp
+++ b/Engine/Source/Developer/UnsavedAssetsTracker/Source/Private/UnsavedAssetsAutoCheckout.cpp
@@ -4,12 +4,15 @@
 
 #include "ISourceControlModule.h"
 #include "Settings/EditorLoadingSavingSettings.h"
+#include "Containers/Ticker.h"
+#include "Misc/Paths.h"
 #include "SourceControlOperations.h"
 #include "UnsavedAssetsTrackerModule.h"
 
 FUnsavedAssetsAutoCheckout::FUnsavedAssetsAutoCheckout(FUnsavedAssetsTrackerModule* Module)
+	: bProcessCheckoutBatchPending(false)
 {
-	Module->OnUnsavedAssetAdded.AddRaw(this, &FUnsavedAssetsAutoCheckout::AsyncCheckout);
+	Module->OnUnsavedAssetAdded.AddRaw(this, &FUnsavedAssetsAutoCheckout::OnAsyncCheckout);
 	AsyncCheckoutComplete.BindRaw(this, &FUnsavedAssetsAutoCheckout::OnAsyncCheckoutComplete);
 }
 
@@ -17,32 +20,93 @@ FUnsavedAssetsAutoCheckout::~FUnsavedAssetsAutoCheckout()
 {
 }
 
-void FUnsavedAssetsAutoCheckout::AsyncCheckout(const FString& AbsoluteAssetFilepath)
+void FUnsavedAssetsAutoCheckout::OnAsyncCheckout(const FString& AbsoluteAssetFilepath)
 {
+	// Check if auto checkout is enabled.
 	const UEditorLoadingSavingSettings* Settings = GetDefault<UEditorLoadingSavingSettings>();
 	if (!Settings->GetAutomaticallyCheckoutOnAssetModification())
 	{
 		return;
 	}
-	ISourceControlProvider& SourceControlProvider = ISourceControlModule::Get().GetProvider();
-	FSourceControlOperationRef CheckOutOperation = ISourceControlOperation::Create<FCheckOut>();
 
-	OperationToPath.Add(&CheckOutOperation.Get()) = AbsoluteAssetFilepath;
-	
-	FUnsavedAssetsTrackerModule::Get().PreUnsavedAssetAutoCheckout.Broadcast(AbsoluteAssetFilepath, CheckOutOperation);
-	
-	SourceControlProvider.Execute(CheckOutOperation, AbsoluteAssetFilepath, EConcurrency::Asynchronous, AsyncCheckoutComplete);
+	// Add to FilesToCheckout batch.
+	// Why are we batching? Because moving a large prefab (1000+ source files) would result in
+	// an equal amount of FCheckOut operation which didn't scale and locked the editor due to
+	// worker thread exhaustion.
+	CheckoutBatch.Add(AbsoluteAssetFilepath);
+
+	// If no batch process tick is pending, queue one.
+	if (!bProcessCheckoutBatchPending)
+	{
+		bProcessCheckoutBatchPending = true;
+
+		FTSTicker::GetCoreTicker().AddTicker(
+			FTickerDelegate::CreateSP(this, &FUnsavedAssetsAutoCheckout::OnProcessCheckoutBatch)
+		);
+	}
 }
 
 void FUnsavedAssetsAutoCheckout::OnAsyncCheckoutComplete(const FSourceControlOperationRef& CheckOutOperation, ECommandResult::Type Result)
 {
-	FString const& AbsoluteAssetFilepath = OperationToPath.FindAndRemoveChecked(&CheckOutOperation.Get());
+	const TArray<FString> FilesInBatch = OperationToPaths.FindAndRemoveChecked(&CheckOutOperation.Get());
+
 	if (Result == ECommandResult::Succeeded)
 	{
-		FUnsavedAssetsTrackerModule::Get().PostUnsavedAssetAutoCheckout.Broadcast(AbsoluteAssetFilepath, CheckOutOperation);	
+		for (const FString& File : FilesInBatch)
+		{
+			FUnsavedAssetsTrackerModule::Get().PostUnsavedAssetAutoCheckout.Broadcast(File, CheckOutOperation);
+		}
+	}
+
+	if (Result == ECommandResult::Failed)
+	{
+		for (const FString& File : FilesInBatch)
+		{
+			FUnsavedAssetsTrackerModule::Get().PostUnsavedAssetAutoCheckoutFailure.Broadcast(File, CheckOutOperation);
+		}
+	}
+}
+
+bool FUnsavedAssetsAutoCheckout::OnProcessCheckoutBatch(float)
+{
+	TArray<FString> FilesToCheckout;
+	FilesToCheckout.Reserve(CheckoutBatch.Num());
+
+	for (const FString& File : CheckoutBatch)
+	{
+		// Why are we checking for file exists?
+		// The OnUnsavedAssetAdded delegate triggers when dragging an asset in the viewport as well,
+		// before it's even placed in the scene or saved to disk. There's no point in attempting
+		// to do an FCheckOut for them.
+		if (FPaths::FileExists(File))
+		{
+			FilesToCheckout.Add(File);
+		}
 	}
-	else if (Result == ECommandResult::Failed)
+	CheckoutBatch.Empty();
+
+	if (FilesToCheckout.Num() > 0)
 	{
-		FUnsavedAssetsTrackerModule::Get().PostUnsavedAssetAutoCheckoutFailure.Broadcast(AbsoluteAssetFilepath, CheckOutOperation);
+		AsyncCheckout(FilesToCheckout);
 	}
+
+	bProcessCheckoutBatchPending = false;
+	return false; // One shot.
 }
+
+void FUnsavedAssetsAutoCheckout::AsyncCheckout(const TArray<FString>& FilesToCheckout)
+{
+	ensure(FilesToCheckout.Num() > 0);
+
+	FSourceControlOperationRef CheckOutOperation = ISourceControlOperation::Create<FCheckOut>();
+
+	OperationToPaths.Add(&CheckOutOperation.Get(), FilesToCheckout);
+
+	for (const FString& File : FilesToCheckout)
+	{
+		FUnsavedAssetsTrackerModule::Get().PreUnsavedAssetAutoCheckout.Broadcast(File, CheckOutOperation);
+	}
+
+	ISourceControlProvider& SourceControlProvider = ISourceControlModule::Get().GetProvider();
+	SourceControlProvider.Execute(CheckOutOperation, FilesToCheckout, EConcurrency::Asynchronous, AsyncCheckoutComplete);
+}
\ No newline at end of file
diff --git a/Engine/Source/Developer/UnsavedAssetsTracker/Source/Private/UnsavedAssetsAutoCheckout.h b/Engine/Source/Developer/UnsavedAssetsTracker/Source/Private/UnsavedAssetsAutoCheckout.h
--- a/Engine/Source/Developer/UnsavedAssetsTracker/Source/Private/UnsavedAssetsAutoCheckout.h
+++ b/Engine/Source/Developer/UnsavedAssetsTracker/Source/Private/UnsavedAssetsAutoCheckout.h
@@ -17,11 +17,20 @@ class FUnsavedAssetsAutoCheckout : public TSharedFromThis<FUnsavedAssetsAutoChec
 public:
 	explicit FUnsavedAssetsAutoCheckout(FUnsavedAssetsTrackerModule* Module);
 	~FUnsavedAssetsAutoCheckout();
+
 private:
-	void AsyncCheckout(const FString& AbsoluteAssetFilepath);
+	void OnAsyncCheckout(const FString& AbsoluteAssetFilepath);
 	
 	FSourceControlOperationComplete AsyncCheckoutComplete;
 	void OnAsyncCheckoutComplete(const FSourceControlOperationRef&, ECommandResult::Type);
 
-	TMap<ISourceControlOperation*, FString> OperationToPath;
+	/* To process checkouts in a single batch once per tick */
+	bool OnProcessCheckoutBatch(float);
+	bool bProcessCheckoutBatchPending;
+
+	void AsyncCheckout(const TArray<FString>& FilesToCheckout);
+
+	TMap<ISourceControlOperation*, TArray<FString>> OperationToPaths;
+
+	TSet<FString> CheckoutBatch;
 };
