Fixed FUnsavedAssetsTracker not to trigger a SyncWithDirtyPackageList if a PIE/game actor is deleted.

diff --git a/Engine/Source/Developer/UnsavedAssetsTracker/Source/Private/UnsavedAssetsTracker.cpp b/Engine/Source/Developer/UnsavedAssetsTracker/Source/Private/UnsavedAssetsTracker.cpp
--- a/Engine/Source/Developer/UnsavedAssetsTracker/Source/Private/UnsavedAssetsTracker.cpp
+++ b/Engine/Source/Developer/UnsavedAssetsTracker/Source/Private/UnsavedAssetsTracker.cpp
@@ -323,6 +323,12 @@ void FUnsavedAssetsTracker::OnPackageDeleted(UPackage* Package)
 
 void FUnsavedAssetsTracker::OnActorDeleted(AActor* Actor)
 {
+	UWorld* World = Actor->GetWorld();
+	if (World && World->IsGameWorld())
+	{
+		return;
+	}
+
 	bSyncWithDirtyPackageList = true;
 }
 
