SourceControl: Fixed UnsavedAssetsAutoCheckout checking out assets that are no longer dirty.

During world recreation, which is triggered by REVERT ALL and SYNC ALL, some assets are temporarily dirty'd by the reload world code because it destroys the old world instance and recreates a new world instance (see: UEditorEngine::NewMap). By the time this code block finishes, the assets are no longer dirty but the events will have been broadcasted by the UnsavedAssetsTracker, resulting in UnsavedAssetsAutoCheckout checking out assets that aren't dirty (anymore).

diff --git a/Engine/Source/Developer/UnsavedAssetsTracker/Source/Private/UnsavedAssetsAutoCheckout.cpp b/Engine/Source/Developer/UnsavedAssetsTracker/Source/Private/UnsavedAssetsAutoCheckout.cpp
--- a/Engine/Source/Developer/UnsavedAssetsTracker/Source/Private/UnsavedAssetsAutoCheckout.cpp
+++ b/Engine/Source/Developer/UnsavedAssetsTracker/Source/Private/UnsavedAssetsAutoCheckout.cpp
@@ -6,6 +6,8 @@
 #include "Settings/EditorLoadingSavingSettings.h"
 #include "Containers/Ticker.h"
 #include "Misc/Paths.h"
+#include "Misc/PackageName.h"
+#include "UObject/Package.h"
 #include "SourceControlOperations.h"
 #include "UnsavedAssetsTrackerModule.h"
 
@@ -78,10 +80,27 @@ bool FUnsavedAssetsAutoCheckout::OnProcessCheckoutBatch(float)
 		// The OnUnsavedAssetAdded delegate triggers when dragging an asset in the viewport as well,
 		// before it's even placed in the scene or saved to disk. There's no point in attempting
 		// to do an FCheckOut for them.
-		if (FPaths::FileExists(File))
+		if (!FPaths::FileExists(File))
 		{
-			FilesToCheckout.Add(File);
+			continue;
 		}
+
+		// Why are we checking if the package is still dirty?
+		// Some packages become temporarily dirty, for example during world recreation.
+		// By the time we try to check them out they may no longer be dirty.
+		FString PackageName;
+		if (FPackageName::TryConvertFilenameToLongPackageName(File, PackageName))
+		{
+			if (UPackage* Package = FindPackage(nullptr, *PackageName))
+			{
+				if (!Package->IsDirty())
+				{
+					continue;
+				}
+			}
+		}
+
+		FilesToCheckout.Add(File);
 	}
 	CheckoutBatch.Empty();
 
