[NO NON-PII COMMIT MESSAGE]

diff --git a/Engine/Source/Developer/UnsavedAssetsTracker/Source/Private/SUnsavedAssetsStatusBarWidget.cpp b/Engine/Source/Developer/UnsavedAssetsTracker/Source/Private/SUnsavedAssetsStatusBarWidget.cpp
--- a/Engine/Source/Developer/UnsavedAssetsTracker/Source/Private/SUnsavedAssetsStatusBarWidget.cpp
+++ b/Engine/Source/Developer/UnsavedAssetsTracker/Source/Private/SUnsavedAssetsStatusBarWidget.cpp
@@ -82,7 +82,15 @@ FText SUnsavedAssetsStatusBarWidget::GetStatusBarTooltip() const
 	}
 	else
 	{
-		return FText::Format(LOCTEXT("Asset_Unsaved_Tooltip", "There are currently {0} assets that are unsaved in this project. Press to initiate save."), UnsavedAssetTracker->GetUnsavedAssetNum());
+		const int32 UnsavedAssetNum = UnsavedAssetTracker->GetUnsavedAssetNum();
+		if (UnsavedAssetNum <= 1)
+		{
+			return FText::Format(LOCTEXT("Single_Asset_Unsaved_Tooltip", "There is currently {0} asset that is unsaved in this project. Press to initiate save."), UnsavedAssetNum);
+		}
+		else
+		{
+			return FText::Format(LOCTEXT("Asset_Unsaved_Tooltip", "There are currently {0} assets that are unsaved in this project. Press to initiate save."), UnsavedAssetNum);
+		}
 	}
 }
 
