SourceControl: Fixed some changes not being immediately undone.

The batching of files into a single FCheckOut operation in #23509995 introduced a subtle bug if a user action resulted in two separate transactions within a single tick. An example of this is adding a component to an actor, which results in 'Add Component' and 'Clicking on Component (tree view)' transactions. Because the PreUnsavedAssetAutoCheckout notification got postponed to the Tick() because of the batching, only the 'Clicking on Component (tree view)' transaction ended up being undone. This fix restores the timing of the PreUnsavedAssetAutoCheckout to its original location and adds a PostUnsavedAssetAutoCheckoutCancel delegate which will be called for assets that do not end up being part of the CheckOutOperation (because they do not exist on disk, or are no longer dirty).

diff --git a/Engine/Source/Developer/UnsavedAssetsTracker/Source/Private/UnsavedAssetsAutoCheckout.cpp b/Engine/Source/Developer/UnsavedAssetsTracker/Source/Private/UnsavedAssetsAutoCheckout.cpp
--- a/Engine/Source/Developer/UnsavedAssetsTracker/Source/Private/UnsavedAssetsAutoCheckout.cpp
+++ b/Engine/Source/Developer/UnsavedAssetsTracker/Source/Private/UnsavedAssetsAutoCheckout.cpp
@@ -13,6 +13,8 @@
 
 FUnsavedAssetsAutoCheckout::FUnsavedAssetsAutoCheckout(FUnsavedAssetsTrackerModule* Module)
 	: bProcessCheckoutBatchPending(false)
+	, CheckOutBatch()
+	, CheckOutOperation(ISourceControlOperation::Create<FCheckOut>())
 {
 	Module->OnUnsavedAssetAdded.AddRaw(this, &FUnsavedAssetsAutoCheckout::OnAsyncCheckout);
 	AsyncCheckoutComplete.BindRaw(this, &FUnsavedAssetsAutoCheckout::OnAsyncCheckoutComplete);
@@ -31,11 +33,15 @@ void FUnsavedAssetsAutoCheckout::OnAsyncCheckout(const FString& AbsoluteAssetFil
 		return;
 	}
 
-	// Add to FilesToCheckout batch.
+	// Add to CheckOutBatch.
 	// Why are we batching? Because moving a large prefab (1000+ source files) would result in
-	// an equal amount of FCheckOut operation which didn't scale and locked the editor due to
+	// an equal amount of FCheckOut operations which didn't scale and locked the editor due to
 	// worker thread exhaustion.
-	CheckoutBatch.Add(AbsoluteAssetFilepath);
+	CheckOutBatch.Add(AbsoluteAssetFilepath);
+
+	// Broadcast that this file is potentially going to be automatically checked out.
+	// This will be followed up by either a Success/Failure/Cancel notification.
+	FUnsavedAssetsTrackerModule::Get().PreUnsavedAssetAutoCheckout.Broadcast(AbsoluteAssetFilepath, CheckOutOperation);
 
 	// If no batch process tick is pending, queue one.
 	if (!bProcessCheckoutBatchPending)
@@ -48,33 +54,40 @@ void FUnsavedAssetsAutoCheckout::OnAsyncCheckout(const FString& AbsoluteAssetFil
 	}
 }
 
-void FUnsavedAssetsAutoCheckout::OnAsyncCheckoutComplete(const FSourceControlOperationRef& CheckOutOperation, ECommandResult::Type Result)
+void FUnsavedAssetsAutoCheckout::OnAsyncCheckoutComplete(const FSourceControlOperationRef& InCheckOutOperation, ECommandResult::Type Result)
 {
-	const TArray<FString> FilesInBatch = OperationToPaths.FindAndRemoveChecked(&CheckOutOperation.Get());
-
-	if (Result == ECommandResult::Succeeded)
+	if (TArray<FString>* Files = OperationToFiles.Find(InCheckOutOperation))
 	{
-		for (const FString& File : FilesInBatch)
+		switch (Result)
 		{
-			FUnsavedAssetsTrackerModule::Get().PostUnsavedAssetAutoCheckout.Broadcast(File, CheckOutOperation);
+		case ECommandResult::Succeeded:
+			NotifySuccess(MakeArrayView(*Files), InCheckOutOperation);
+			break;
+		case ECommandResult::Failed:
+			NotifyFailure(MakeArrayView(*Files), InCheckOutOperation);
+			break;
+		case ECommandResult::Cancelled:
+			NotifyCancel(MakeArrayView(*Files), InCheckOutOperation);
+			break;
+		default:
+			checkNoEntry();
+			break;
 		}
-	}
 
-	if (Result == ECommandResult::Failed)
-	{
-		for (const FString& File : FilesInBatch)
-		{
-			FUnsavedAssetsTrackerModule::Get().PostUnsavedAssetAutoCheckoutFailure.Broadcast(File, CheckOutOperation);
-		}
+		OperationToFiles.Remove(InCheckOutOperation);
 	}
 }
 
 bool FUnsavedAssetsAutoCheckout::OnProcessCheckoutBatch(float)
 {
 	TArray<FString> FilesToCheckout;
-	FilesToCheckout.Reserve(CheckoutBatch.Num());
+	FilesToCheckout.Reserve(CheckOutBatch.Num());
 
-	for (const FString& File : CheckoutBatch)
+	TArray<FString> FilesToCancel;
+	FilesToCancel.Reserve(CheckOutBatch.Num());
+
+	// Determine which files need not be checked out and can thus be canceled.
+	for (const FString& File : CheckOutBatch)
 	{
 		// Why are we checking for file exists?
 		// The OnUnsavedAssetAdded delegate triggers when dragging an asset in the viewport as well,
@@ -82,6 +95,7 @@ bool FUnsavedAssetsAutoCheckout::OnProcessCheckoutBatch(float)
 		// to do an FCheckOut for them.
 		if (!FPaths::FileExists(File))
 		{
+			FilesToCancel.Add(File);
 			continue;
 		}
 
@@ -95,37 +109,65 @@ bool FUnsavedAssetsAutoCheckout::OnProcessCheckoutBatch(float)
 			{
 				if (!Package->IsDirty())
 				{
+					FilesToCancel.Add(File);
 					continue;
 				}
 			}
 		}
 
+		// If not canceled, it should be checked out.
 		FilesToCheckout.Add(File);
 	}
-	CheckoutBatch.Empty();
 
+	// Call the Cancel delegate for those that no longer need to be checked out.
+	if (FilesToCancel.Num() > 0)
+	{
+		NotifyCancel(FilesToCancel, CheckOutOperation);
+	}
+	// Call the Success/Failure/Cancel delegate for others by triggering a checkout operation.
 	if (FilesToCheckout.Num() > 0)
 	{
-		AsyncCheckout(FilesToCheckout);
+		TriggerAsyncCheckout(FilesToCheckout, CheckOutOperation);
 	}
 
+	// Reset now that all files in this batch are either cancelled or pending checkout.
+	CheckOutBatch.Empty();
+	CheckOutOperation = ISourceControlOperation::Create<FCheckOut>();
 	bProcessCheckoutBatchPending = false;
+
 	return false; // One shot.
 }
 
-void FUnsavedAssetsAutoCheckout::AsyncCheckout(const TArray<FString>& FilesToCheckout)
+void FUnsavedAssetsAutoCheckout::TriggerAsyncCheckout(const TArray<FString>& FilesToCheckout, const FSourceControlOperationRef& Operation)
 {
 	ensure(FilesToCheckout.Num() > 0);
 
-	FSourceControlOperationRef CheckOutOperation = ISourceControlOperation::Create<FCheckOut>();
+	OperationToFiles.Add(Operation, FilesToCheckout);
 
-	OperationToPaths.Add(&CheckOutOperation.Get(), FilesToCheckout);
+	ISourceControlProvider& SourceControlProvider = ISourceControlModule::Get().GetProvider();
+	SourceControlProvider.Execute(Operation, FilesToCheckout, EConcurrency::Asynchronous, AsyncCheckoutComplete);
+}
 
-	for (const FString& File : FilesToCheckout)
+void FUnsavedAssetsAutoCheckout::NotifySuccess(TConstArrayView<FString> Files, const FSourceControlOperationRef& Operation)
+{
+	for (const FString& File : Files)
 	{
-		FUnsavedAssetsTrackerModule::Get().PreUnsavedAssetAutoCheckout.Broadcast(File, CheckOutOperation);
+		FUnsavedAssetsTrackerModule::Get().PostUnsavedAssetAutoCheckout.Broadcast(File, Operation);
 	}
+}
 
-	ISourceControlProvider& SourceControlProvider = ISourceControlModule::Get().GetProvider();
-	SourceControlProvider.Execute(CheckOutOperation, FilesToCheckout, EConcurrency::Asynchronous, AsyncCheckoutComplete);
+void FUnsavedAssetsAutoCheckout::NotifyFailure(TConstArrayView<FString> Files, const FSourceControlOperationRef& Operation)
+{
+	for (const FString& File : Files)
+	{
+		FUnsavedAssetsTrackerModule::Get().PostUnsavedAssetAutoCheckoutFailure.Broadcast(File, Operation);
+	}
+}
+
+void FUnsavedAssetsAutoCheckout::NotifyCancel(TConstArrayView<FString> Files, const FSourceControlOperationRef& Operation)
+{
+	for (const FString& File : Files)
+	{
+		FUnsavedAssetsTrackerModule::Get().PostUnsavedAssetAutoCheckoutCancel.Broadcast(File, Operation);
+	}
 }
\ No newline at end of file
diff --git a/Engine/Source/Developer/UnsavedAssetsTracker/Source/Private/UnsavedAssetsAutoCheckout.h b/Engine/Source/Developer/UnsavedAssetsTracker/Source/Private/UnsavedAssetsAutoCheckout.h
--- a/Engine/Source/Developer/UnsavedAssetsTracker/Source/Private/UnsavedAssetsAutoCheckout.h
+++ b/Engine/Source/Developer/UnsavedAssetsTracker/Source/Private/UnsavedAssetsAutoCheckout.h
@@ -5,6 +5,7 @@
 #include "Templates/SharedPointer.h"
 #include "Containers/Map.h"
 #include "ISourceControlProvider.h"
+#include "SourceControlOperations.h"
 
 class ISourceControlOperation;
 class FUnsavedAssetsTrackerModule;
@@ -28,9 +29,14 @@ private:
 	bool OnProcessCheckoutBatch(float);
 	bool bProcessCheckoutBatchPending;
 
-	void AsyncCheckout(const TArray<FString>& FilesToCheckout);
+	void TriggerAsyncCheckout(const TArray<FString>& FilesToCheckout, const FSourceControlOperationRef& Operation);
 
-	TMap<ISourceControlOperation*, TArray<FString>> OperationToPaths;
+	void NotifySuccess(TConstArrayView<FString> Files, const FSourceControlOperationRef& Operation);
+	void NotifyFailure(TConstArrayView<FString> Files, const FSourceControlOperationRef& Operation);
+	void NotifyCancel(TConstArrayView<FString> Files, const FSourceControlOperationRef& Operation);
 
-	TSet<FString> CheckoutBatch;
+	TMap<FSourceControlOperationRef, TArray<FString>> OperationToFiles;
+
+	TSet<FString> CheckOutBatch;
+	FSourceControlOperationRef CheckOutOperation;
 };
diff --git a/Engine/Source/Developer/UnsavedAssetsTracker/Source/Public/UnsavedAssetsTrackerModule.h b/Engine/Source/Developer/UnsavedAssetsTracker/Source/Public/UnsavedAssetsTrackerModule.h
--- a/Engine/Source/Developer/UnsavedAssetsTracker/Source/Public/UnsavedAssetsTrackerModule.h
+++ b/Engine/Source/Developer/UnsavedAssetsTracker/Source/Public/UnsavedAssetsTrackerModule.h
@@ -65,7 +65,11 @@ public:
 	FOnUnsavedAssetRemoved OnUnsavedAssetRemoved;
 
 	/** Invoked before an asset is potentially automatically checked out of source control */
+	/** Every notified File/Operation combination is followed by one of the Cancel/Success/Failure notifications below */
 	FPreUnsavedAssetAutoCheckout PreUnsavedAssetAutoCheckout;
+
+	/** Invoked after an asset has been cancelled for automatic check out */
+	FPostUnsavedAssetAutoCheckout PostUnsavedAssetAutoCheckoutCancel;
 	
 	/** Invoked after an asset has been checked out of source control */
 	FPostUnsavedAssetAutoCheckout PostUnsavedAssetAutoCheckout;
