SourceControl: If SourceControl is not enabled or not available, don't attempt to auto FCheckOut any files, as it could prevent the user from making any edits.

Currently, if Skein is enabled but unavailable (for example: the user is disconnected from the internet), any action performed by the user will be Auto-Undo'd because the UnsavedAssetsAutoCheckout instance will trigger FCheckOut operations that will all fail.
This can be fixed in the Skein code as well but it probably makes sense to detect this on the UnsavedAssetsAutoCheckout level already.

diff --git a/Engine/Source/Developer/UnsavedAssetsTracker/Source/Private/UnsavedAssetsAutoCheckout.cpp b/Engine/Source/Developer/UnsavedAssetsTracker/Source/Private/UnsavedAssetsAutoCheckout.cpp
--- a/Engine/Source/Developer/UnsavedAssetsTracker/Source/Private/UnsavedAssetsAutoCheckout.cpp
+++ b/Engine/Source/Developer/UnsavedAssetsTracker/Source/Private/UnsavedAssetsAutoCheckout.cpp
@@ -10,6 +10,7 @@
 #include "UObject/Package.h"
 #include "SourceControlOperations.h"
 #include "UnsavedAssetsTrackerModule.h"
+#include "UncontrolledChangelistsModule.h"
 
 FUnsavedAssetsAutoCheckout::FUnsavedAssetsAutoCheckout(FUnsavedAssetsTrackerModule* Module)
 	: bProcessCheckoutBatchPending(false)
@@ -33,6 +34,30 @@ void FUnsavedAssetsAutoCheckout::OnAsyncCheckout(const FString& AbsoluteAssetFil
 		return;
 	}
 
+	// Check if SourceControl is enabled and available.
+	ISourceControlModule& SourceControlModule = ISourceControlModule::Get();
+	if (!SourceControlModule.IsEnabled())
+	{
+		return;
+	}
+	else
+	{
+		bool bIsAvailable = SourceControlModule.GetProvider().IsAvailable();
+		if (!bIsAvailable)
+		{
+			// The UncontrolledChangelists module relies on the 'checkout' to make the files locally writable
+			// and able to reconcile later. So if used and enabled, allow the 'checkout' to continue.
+			bool bUncontrolledChangelistsUsed = SourceControlModule.GetProvider().UsesUncontrolledChangelists();
+			bool bUncontrolledChangelistsEnabled = FUncontrolledChangelistsModule::Get().IsEnabled();
+			
+			bool bUncontrolledChangelistsRequiresCheckout = bUncontrolledChangelistsUsed && bUncontrolledChangelistsEnabled;
+			if (!bUncontrolledChangelistsRequiresCheckout)
+			{
+				return;
+			}
+		}
+	}
+
 	// Add to CheckOutBatch.
 	// Why are we batching? Because moving a large prefab (1000+ source files) would result in
 	// an equal amount of FCheckOut operations which didn't scale and locked the editor due to
diff --git a/Engine/Source/Developer/UnsavedAssetsTracker/Source/UnsavedAssetsTracker.Build.cs b/Engine/Source/Developer/UnsavedAssetsTracker/Source/UnsavedAssetsTracker.Build.cs
--- a/Engine/Source/Developer/UnsavedAssetsTracker/Source/UnsavedAssetsTracker.Build.cs
+++ b/Engine/Source/Developer/UnsavedAssetsTracker/Source/UnsavedAssetsTracker.Build.cs
@@ -13,6 +13,7 @@ public class UnsavedAssetsTracker : ModuleRules
 				"Engine",
 				"EditorSubsystem",
 				"SourceControl",
+				"UncontrolledChangelists",
 			}
 		);
 
