Unsaved Asset Tracker
- Add PostGC callback to rebuild list (make sure tracker stops tracking GCed packages)
- Tracker rebuild list delayed while in PIE

#rnx

diff --git a/Engine/Source/Developer/UnsavedAssetsTracker/Source/Private/UnsavedAssetsTracker.cpp b/Engine/Source/Developer/UnsavedAssetsTracker/Source/Private/UnsavedAssetsTracker.cpp
--- a/Engine/Source/Developer/UnsavedAssetsTracker/Source/Private/UnsavedAssetsTracker.cpp
+++ b/Engine/Source/Developer/UnsavedAssetsTracker/Source/Private/UnsavedAssetsTracker.cpp
@@ -23,6 +23,7 @@
 #include "LevelEditor.h"
 #include "ProfilingDebugging/CpuProfilerTrace.h"
 #include "IAssetTools.h"
+#include "UObject/UObjectGlobals.h"
 
 #define LOCTEXT_NAMESPACE "UnsavedAssetsTracker"
 
@@ -148,6 +149,9 @@ FUnsavedAssetsTracker::FUnsavedAssetsTracker()
 	// Register for the package modified callback to catch packages that have been saved
 	UPackage::PackageSavedWithContextEvent.AddRaw(this, &FUnsavedAssetsTracker::OnPackageSaved);
 
+	// Register for post gc to cleanup package list that might have been gced
+	FCoreUObjectDelegates::GetPostGarbageCollect().AddRaw(this, &FUnsavedAssetsTracker::OnPostGarbageCollect);
+
 	// Register to get notified when something gets deleted
 	FEditorDelegates::OnPackageDeleted.AddRaw(this, &FUnsavedAssetsTracker::OnPackageDeleted);
 	if (GEngine)
@@ -182,6 +186,7 @@ FUnsavedAssetsTracker::~FUnsavedAssetsTracker()
 	UPackage::PackageMarkedDirtyEvent.RemoveAll(this);
 	UPackage::PackageSavedWithContextEvent.RemoveAll(this);
 	FEditorDelegates::OnPackageDeleted.RemoveAll(this);
+	FCoreUObjectDelegates::GetPostGarbageCollect().RemoveAll(this);
 
 	if (GEngine)
 	{
@@ -321,14 +326,13 @@ void FUnsavedAssetsTracker::OnPackageDeleted(UPackage* Package)
 	}
 }
 
-void FUnsavedAssetsTracker::OnActorDeleted(AActor* Actor)
+void FUnsavedAssetsTracker::OnPostGarbageCollect()
 {
-	UWorld* World = Actor->GetWorld();
-	if (World && World->IsGameWorld())
-	{
-		return;
-	}
+	bSyncWithDirtyPackageList = true;
+}
 
+void FUnsavedAssetsTracker::OnActorDeleted(AActor* Actor)
+{
 	bSyncWithDirtyPackageList = true;
 }
 
@@ -615,6 +619,11 @@ bool FUnsavedAssetsTracker::Tick(float DeltaTime)
 
 	if (bSyncWithDirtyPackageList)
 	{
+		if (GEditor && GEditor->GetPIEWorldContext() != nullptr)
+		{
+			return true; // delay update until PIE ends
+		}
+		
 		SyncWithDirtyPackageList();
 		bSyncWithDirtyPackageList = false;
 	}
diff --git a/Engine/Source/Developer/UnsavedAssetsTracker/Source/Private/UnsavedAssetsTracker.h b/Engine/Source/Developer/UnsavedAssetsTracker/Source/Private/UnsavedAssetsTracker.h
--- a/Engine/Source/Developer/UnsavedAssetsTracker/Source/Private/UnsavedAssetsTracker.h
+++ b/Engine/Source/Developer/UnsavedAssetsTracker/Source/Private/UnsavedAssetsTracker.h
@@ -100,6 +100,9 @@ private:
 	void OnUndo(const FTransactionContext& TransactionContext, bool Succeeded);
 	void OnRedo(const FTransactionContext& TransactionContext, bool Succeeded);
 
+	/** Invoked after garbage collection. */
+	void OnPostGarbageCollect();
+
 	/** Invoked when a package is deleted. */
 	void OnPackageDeleted(UPackage* Package);
 
