SourceControl: Added 'uses checkout' check to UnsavedAssetsAutoCheckout.

#rnx

diff --git a/Engine/Source/Developer/UnsavedAssetsTracker/Source/Private/UnsavedAssetsAutoCheckout.cpp b/Engine/Source/Developer/UnsavedAssetsTracker/Source/Private/UnsavedAssetsAutoCheckout.cpp
--- a/Engine/Source/Developer/UnsavedAssetsTracker/Source/Private/UnsavedAssetsAutoCheckout.cpp
+++ b/Engine/Source/Developer/UnsavedAssetsTracker/Source/Private/UnsavedAssetsAutoCheckout.cpp
@@ -58,6 +58,13 @@ void FUnsavedAssetsAutoCheckout::OnAsyncCheckout(const FString& AbsoluteAssetFil
 		}
 	}
 
+	// Check if SourceControl is using checkout.
+	bool bUsesCheckout = SourceControlModule.GetProvider().UsesCheckout();
+	if (!bUsesCheckout)
+	{
+		return;
+	}
+
 	// Add to CheckOutBatch.
 	// Why are we batching? Because moving a large prefab (1000+ source files) would result in
 	// an equal amount of FCheckOut operations which didn't scale and locked the editor due to
