SourceControl: Don't attempt to auto-checkout a file that's already checked out by active user.

#rnx

diff --git a/Engine/Source/Developer/UnsavedAssetsTracker/Source/Private/UnsavedAssetsAutoCheckout.cpp b/Engine/Source/Developer/UnsavedAssetsTracker/Source/Private/UnsavedAssetsAutoCheckout.cpp
--- a/Engine/Source/Developer/UnsavedAssetsTracker/Source/Private/UnsavedAssetsAutoCheckout.cpp
+++ b/Engine/Source/Developer/UnsavedAssetsTracker/Source/Private/UnsavedAssetsAutoCheckout.cpp
@@ -65,6 +65,14 @@ void FUnsavedAssetsAutoCheckout::OnAsyncCheckout(const FString& AbsoluteAssetFil
 		return;
 	}
 
+	// Check if SourceControl has already checked out this file for us.
+	// If it's checked out by someone else, we want this to proceed, so the operation can be auto-undo'd.
+	FSourceControlStatePtr State = SourceControlModule.GetProvider().GetState(AbsoluteAssetFilepath, EStateCacheUsage::Use);
+	if (State.IsValid() && State->IsCheckedOut())
+	{
+		return;
+	}
+
 	// Add to CheckOutBatch.
 	// Why are we batching? Because moving a large prefab (1000+ source files) would result in
 	// an equal amount of FCheckOut operations which didn't scale and locked the editor due to
