UnsavedAssetsAutoCheckout: Do not attempt to run FCheckout on added files. It either fails (and then get's undone which is not what we want - it prevents the user from deleting the added asset, URC) or is unnecessary (Perforce).

Note that we cannot check ISourceControlState::CanCheckout(), as we WANT FCheckout to run (and fail) in case the asset is checked out by another user or not at head, after which the action can be correctly undone.

(need to get this in pre hard lock)

#rnx

diff --git a/Engine/Source/Developer/UnsavedAssetsTracker/Source/Private/UnsavedAssetsAutoCheckout.cpp b/Engine/Source/Developer/UnsavedAssetsTracker/Source/Private/UnsavedAssetsAutoCheckout.cpp
--- a/Engine/Source/Developer/UnsavedAssetsTracker/Source/Private/UnsavedAssetsAutoCheckout.cpp
+++ b/Engine/Source/Developer/UnsavedAssetsTracker/Source/Private/UnsavedAssetsAutoCheckout.cpp
@@ -72,6 +72,11 @@ void FUnsavedAssetsAutoCheckout::OnAsyncCheckout(const FString& AbsoluteAssetFil
 	{
 		return;
 	}
+	// Check if this is an added file, in which case it cannot be checked out or does not need to be checked out.
+	if (State.IsValid() && State->IsAdded())
+	{
+		return;
+	}
 
 	// Add to CheckOutBatch.
 	// Why are we batching? Because moving a large prefab (1000+ source files) would result in
