Unsaved Assets Tracker: Quick change to remove columns one-by-one instead of using batch removal which has some known issues causing crashes

diff --git a/Engine/Source/Developer/UnsavedAssetsTracker/Source/Private/UnsavedAssetsTracker.cpp b/Engine/Source/Developer/UnsavedAssetsTracker/Source/Private/UnsavedAssetsTracker.cpp
--- a/Engine/Source/Developer/UnsavedAssetsTracker/Source/Private/UnsavedAssetsTracker.cpp
+++ b/Engine/Source/Developer/UnsavedAssetsTracker/Source/Private/UnsavedAssetsTracker.cpp
@@ -524,10 +524,10 @@ void FUnsavedAssetsTracker::UpdateTedsPackage(const FString& PackagePathname, co
 					Storage->WalkDepthFirst(Storage->FindHierarchyByName(PackageHierarchyName), PackageRow, GetChildRows);
 
 					// Remove package columns from any rows that don't belong to this package anymore so they can resync
-					Storage->BatchAddRemoveColumns(
-						TConstArrayView<RowHandle>(RemapChildrenRowHandles),
-						{},
-						{FTypedElementPackageReference::StaticStruct(), FTedsPackageDirtyTag::StaticStruct(), FTedsPrimaryPackageObjectTag::StaticStruct()});
+					for (RowHandle RemapRow : RemapChildrenRowHandles)
+					{
+						Storage->RemoveColumns(RemapRow, {FTypedElementPackageReference::StaticStruct(), FTedsPackageDirtyTag::StaticStruct(), FTedsPrimaryPackageObjectTag::StaticStruct()});
+					}
 
 					// Any rows still in the package should update to the correct IsDirty state
 					for (const RowHandle ChildRow : ChildrenRowHandles)
