[Unsaved Asset Tracker] Add back the batch remove Column functionality in Main since this branch has the batch fixes

#rnx
#Synced-CL 51342861

diff --git a/Engine/Source/Developer/UnsavedAssetsTracker/Source/Private/UnsavedAssetsTracker.cpp b/Engine/Source/Developer/UnsavedAssetsTracker/Source/Private/UnsavedAssetsTracker.cpp
--- a/Engine/Source/Developer/UnsavedAssetsTracker/Source/Private/UnsavedAssetsTracker.cpp
+++ b/Engine/Source/Developer/UnsavedAssetsTracker/Source/Private/UnsavedAssetsTracker.cpp
@@ -524,10 +524,11 @@ void FUnsavedAssetsTracker::UpdateTedsPackage(const FString& PackagePathname, co
 					Storage->WalkDepthFirst(Storage->FindHierarchyByName(PackageHierarchyName), PackageRow, GetChildRows);
 
 					// Remove package columns from any rows that don't belong to this package anymore so they can resync
-					for (RowHandle RemapRow : RemapChildrenRowHandles)
-					{
-						Storage->RemoveColumns(RemapRow, {FTypedElementPackageReference::StaticStruct(), FTedsPackageDirtyTag::StaticStruct(), FTedsPrimaryPackageObjectTag::StaticStruct()});
-					}
+					Storage->BatchAddRemoveColumns(
+						TConstArrayView<RowHandle>(RemapChildrenRowHandles),
+						{},
+						{FTypedElementPackageReference::StaticStruct(), FTedsPackageDirtyTag::StaticStruct(), FTedsPrimaryPackageObjectTag::StaticStruct()});
+					
 
 					// Any rows still in the package should update to the correct IsDirty state
 					for (const RowHandle ChildRow : ChildrenRowHandles)
@@ -567,10 +568,6 @@ void FUnsavedAssetsTracker::RemoveTedsPackage(const FString& PackagePathname)
 						};
 					Storage->WalkDepthFirst(Storage->FindHierarchyByName(PackageHierarchyName), PackageRow, GetChildRows);
 
-					for (const RowHandle ChildRow : ChildRowHandles)
-					{
-						Storage->RemoveColumns(ChildRow, {FTypedElementPackageReference::StaticStruct(), FTedsPackageDirtyTag::StaticStruct(), FTedsPrimaryPackageObjectTag::StaticStruct()});
-					}
 					Storage->BatchAddRemoveColumns(
 						TConstArrayView<RowHandle>(ChildRowHandles),
 						{},
