[TEDS Revision Control] Move package row mapping to use extensions rather than the LongPackageName string so it can match the revision control mapping that needs the extensions

#Synced-CL 51412100

diff --git a/Engine/Plugins/Experimental/EditorDataStorageFeatures/Source/TedsRevisionControl/Private/Queries/ObjectPackagePathToColumnQueries.cpp b/Engine/Plugins/Experimental/EditorDataStorageFeatures/Source/TedsRevisionControl/Private/Queries/ObjectPackagePathToColumnQueries.cpp
--- a/Engine/Plugins/Experimental/EditorDataStorageFeatures/Source/TedsRevisionControl/Private/Queries/ObjectPackagePathToColumnQueries.cpp
+++ b/Engine/Plugins/Experimental/EditorDataStorageFeatures/Source/TedsRevisionControl/Private/Queries/ObjectPackagePathToColumnQueries.cpp
@@ -182,7 +182,8 @@ void UTypedElementUObjectPackagePathFactory::RegisterTryAddPackageRef(UE::Editor
 			{
 				FString Path;
 				Target->GetPathName(nullptr, Path);
-				if (FString PackageFilename; FPackageName::TryConvertLongPackageNameToFilename(Path, PackageFilename))
+				FString Extension = Target->ContainsMap() ? TEXT(".umap") : TEXT(".uasset");
+				if (FString PackageFilename; FPackageName::TryConvertLongPackageNameToFilename(Path, PackageFilename, Extension))
 				{
 					FPaths::NormalizeFilename(PackageFilename);
 					FString FullPackageFilename = FPaths::ConvertRelativePathToFull(PackageFilename);
diff --git a/Engine/Source/Developer/UnsavedAssetsTracker/Source/Private/UnsavedAssetsTracker.cpp b/Engine/Source/Developer/UnsavedAssetsTracker/Source/Private/UnsavedAssetsTracker.cpp
--- a/Engine/Source/Developer/UnsavedAssetsTracker/Source/Private/UnsavedAssetsTracker.cpp
+++ b/Engine/Source/Developer/UnsavedAssetsTracker/Source/Private/UnsavedAssetsTracker.cpp
@@ -488,10 +488,11 @@ void FUnsavedAssetsTracker::UpdateTedsPackage(const FString& PackagePathname, co
 	{
 		if (ICoreProvider* Storage = GetMutableDataStorageFeature<ICoreProvider>(StorageFeatureName))
 		{
-			if (FString FullPackagePath; FPackageName::TryConvertLongPackageNameToFilename(PackagePathname, FullPackagePath))
+			if (FString PackagePath; FPackageName::DoesPackageExist(PackagePathname, &PackagePath))
 			{
-				FullPackagePath = FPaths::ConvertRelativePathToFull(FullPackagePath);
-				if (const RowHandle PackageRow = Storage->LookupMappedRow(UE::Editor::RevisionControl::MappingDomain, FMapKey(FullPackagePath)); Storage->IsRowAvailable(PackageRow))
+				FPaths::NormalizeFilename(PackagePath);
+				FString FullPackageFilename = FPaths::ConvertRelativePathToFull(PackagePath);
+				if (const RowHandle PackageRow = Storage->LookupMappedRow(UE::Editor::RevisionControl::MappingDomain, FMapKey(FullPackageFilename)); Storage->IsRowAvailable(PackageRow))
 				{
 					TArray<RowHandle> ChildrenRowHandles; // Rows that remain in the package, but need to sync IsDirty
 					TArray<RowHandle> RemapChildrenRowHandles; // Rows that are no longer in the package and need to be resynced in the processor
@@ -556,10 +557,11 @@ void FUnsavedAssetsTracker::RemoveTedsPackage(const FString& PackagePathname)
 	{
 		if (ICoreProvider* Storage = GetMutableDataStorageFeature<ICoreProvider>(StorageFeatureName))
 		{
-			if (FString FullPackagePath; FPackageName::TryConvertLongPackageNameToFilename(PackagePathname, FullPackagePath))
+			if (FString PackagePath; FPackageName::DoesPackageExist(PackagePathname, &PackagePath))
 			{
-				FullPackagePath = FPaths::ConvertRelativePathToFull(FullPackagePath);
-				if (const RowHandle PackageRow = Storage->LookupMappedRow(UE::Editor::RevisionControl::MappingDomain, FMapKey(FullPackagePath)); Storage->IsRowAvailable(PackageRow))
+				FPaths::NormalizeFilename(PackagePath);
+				FString FullPackageFilename = FPaths::ConvertRelativePathToFull(PackagePath);
+				if (const RowHandle PackageRow = Storage->LookupMappedRow(UE::Editor::RevisionControl::MappingDomain, FMapKey(FullPackageFilename)); Storage->IsRowAvailable(PackageRow))
 				{
 					TArray<RowHandle> ChildRowHandles;
 					auto GetChildRows = [&ChildRowHandles](const ICoreProvider&, RowHandle, RowHandle TargetRow)
